Adding WebBrowser Developer project

Web Browser project mostly exposes interfaces, hiding all of the CEF3 implementation and technically allowing us to slot in another web system if necessary. Normal process will be to request a browser window interface from the web browser singleton. There is also an implementation of ISlateViewport so that it can be displayed via an SViewport.
The UnrealCEFSubProcess executable is needed for CEF3 to run additional processes for rendering etc. I don't expect this to change very often so it should hopefully just be built once and distributed, since we can't currently specify a program dependency for other modules.
Added a CEF3Utils project to share .dll loading code between web browser and the separate sub process executable.

diff --git a/Engine/Source/Developer/CEF3Utils/CEF3Utils.Build.cs b/Engine/Source/Developer/CEF3Utils/CEF3Utils.Build.cs
new file mode 100644
--- /dev/null
+++ b/Engine/Source/Developer/CEF3Utils/CEF3Utils.Build.cs
@@ -0,0 +1,26 @@
+// Copyright 1998-2014 Epic Games, Inc. All Rights Reserved.
+
+using UnrealBuildTool;
+
+public class CEF3Utils : ModuleRules
+{
+	public CEF3Utils(TargetInfo Target)
+	{
+		PublicIncludePaths.Add("Developer/CEF3Utils/Public");
+		PrivateIncludePaths.Add("Developer/CEF3Utils/Private");
+
+		PrivateDependencyModuleNames.AddRange(
+			new string[]
+			{
+				"Core",
+			}
+		);
+
+		if (Target.Platform == UnrealTargetPlatform.Win64)
+		{
+			AddThirdPartyPrivateStaticDependencies(Target,
+				"CEF3"
+				);
+		}
+	}
+}
diff --git a/Engine/Source/Developer/CEF3Utils/Private/CEF3Utils.cpp b/Engine/Source/Developer/CEF3Utils/Private/CEF3Utils.cpp
new file mode 100644
--- /dev/null
+++ b/Engine/Source/Developer/CEF3Utils/Private/CEF3Utils.cpp
@@ -0,0 +1,66 @@
+// Copyright 1998-2014 Epic Games, Inc. All Rights Reserved.
+
+#include "CEF3UtilsPrivatePCH.h"
+
+DEFINE_LOG_CATEGORY(LogCEF3Utils);
+
+IMPLEMENT_MODULE(FDefaultModuleImpl, CEF3);
+
+#if WITH_CEF3
+namespace CEF3Utils
+{
+#if PLATFORM_WINDOWS
+	void* CEF3DLLHandle = nullptr;
+	void* D3DHandle = nullptr;
+	void* MPEGHandle = nullptr;
+	void* ICUDTHandle = nullptr;
+	void* GLESHandle = nullptr;
+	void* EGLHandle = nullptr;
+#endif
+
+	void LoadCEF3Modules()
+	{
+#if PLATFORM_WINDOWS
+		FString DllPath(FPaths::Combine(*FPaths::EngineDir(), TEXT("Binaries/ThirdParty/CEF3/Win64")));
+
+		CEF3DLLHandle = FPlatformProcess::GetDllHandle(*FPaths::Combine(*DllPath, TEXT("libcef.dll")));
+		if (CEF3DLLHandle == NULL)
+		{
+			int32 ErrorNum = FPlatformMisc::GetLastError();
+			TCHAR ErrorMsg[1024];
+			FPlatformMisc::GetSystemErrorMessage(ErrorMsg, 1024, ErrorNum);
+			UE_LOG(LogCEF3Utils, Error, TEXT("Failed to get CEF3 DLL handle: %s (%d)"), ErrorMsg, ErrorNum);
+		}
+
+	#if WINVER >= 0x600 // Different dll used pre-Vista
+		D3DHandle = FPlatformProcess::GetDllHandle(*FPaths::Combine(*DllPath, TEXT("d3dcompiler_46.dll")));
+	#else
+		D3DHandle = FPlatformProcess::GetDllHandle(*FPaths::Combine(*DllPath, TEXT("d3dcompiler_43.dll")));
+	#endif
+		MPEGHandle = FPlatformProcess::GetDllHandle(*FPaths::Combine(*DllPath, TEXT("ffmpegsumo.dll")));
+		ICUDTHandle = FPlatformProcess::GetDllHandle(*FPaths::Combine(*DllPath, TEXT("icudt.dll")));
+		GLESHandle = FPlatformProcess::GetDllHandle(*FPaths::Combine(*DllPath, TEXT("libGLESv2.dll")));
+		EGLHandle = FPlatformProcess::GetDllHandle(*FPaths::Combine(*DllPath, TEXT("libEGL.dll")));
+#endif
+	}
+
+	void UnloadCEF3Modules()
+	{
+#if PLATFORM_WINDOWS
+		FPlatformProcess::FreeDllHandle(CEF3DLLHandle);
+		CEF3DLLHandle = nullptr;
+
+		FPlatformProcess::FreeDllHandle(D3DHandle);
+		D3DHandle = nullptr;
+		FPlatformProcess::FreeDllHandle(MPEGHandle);
+		MPEGHandle = nullptr;
+		FPlatformProcess::FreeDllHandle(ICUDTHandle);
+		ICUDTHandle = nullptr;
+		FPlatformProcess::FreeDllHandle(GLESHandle);
+		GLESHandle = nullptr;
+		FPlatformProcess::FreeDllHandle(EGLHandle);
+		EGLHandle = nullptr;
+#endif
+	}
+};
+#endif //WITH_CEF3
diff --git a/Engine/Source/Developer/CEF3Utils/Private/CEF3UtilsPrivatePCH.h b/Engine/Source/Developer/CEF3Utils/Private/CEF3UtilsPrivatePCH.h
new file mode 100644
--- /dev/null
+++ b/Engine/Source/Developer/CEF3Utils/Private/CEF3UtilsPrivatePCH.h
@@ -0,0 +1,9 @@
+// Copyright 1998-2014 Epic Games, Inc. All Rights Reserved.
+
+#pragma once
+
+#include "Core.h"
+#include "ModuleManager.h"
+#include "CEF3Utils.h"
+
+DECLARE_LOG_CATEGORY_EXTERN(LogCEF3Utils, Log, All);
\ No newline at end of file
diff --git a/Engine/Source/Developer/CEF3Utils/Public/CEF3Utils.h b/Engine/Source/Developer/CEF3Utils/Public/CEF3Utils.h
new file mode 100644
--- /dev/null
+++ b/Engine/Source/Developer/CEF3Utils/Public/CEF3Utils.h
@@ -0,0 +1,18 @@
+// Copyright 1998-2014 Epic Games, Inc. All Rights Reserved.
+
+#pragma once
+
+#if WITH_CEF3
+namespace CEF3Utils
+{
+	/**
+	 * Load the required modules for CEF3
+	 */
+	CEF3UTILS_API void LoadCEF3Modules();
+
+	/**
+	 * Unload the required modules for CEF3
+	 */
+	CEF3UTILS_API void UnloadCEF3Modules();
+};
+#endif //WITH_CEF3
diff --git a/Engine/Source/Developer/WebBrowser/Private/WebBrowserApp.cpp b/Engine/Source/Developer/WebBrowser/Private/WebBrowserApp.cpp
new file mode 100644
--- /dev/null
+++ b/Engine/Source/Developer/WebBrowser/Private/WebBrowserApp.cpp
@@ -0,0 +1,16 @@
+// Copyright 1998-2014 Epic Games, Inc. All Rights Reserved.
+
+#include "WebBrowserPrivatePCH.h"
+#include "WebBrowserApp.h"
+
+#if WITH_CEF3
+FWebBrowserApp::FWebBrowserApp()
+{
+}
+
+void FWebBrowserApp::OnContextInitialized()
+{
+	// Called on the browser process UI thread immediately after the CEF context
+	// has been initialized.
+}
+#endif
diff --git a/Engine/Source/Developer/WebBrowser/Private/WebBrowserApp.h b/Engine/Source/Developer/WebBrowser/Private/WebBrowserApp.h
new file mode 100644
--- /dev/null
+++ b/Engine/Source/Developer/WebBrowser/Private/WebBrowserApp.h
@@ -0,0 +1,39 @@
+// Copyright 1998-2014 Epic Games, Inc. All Rights Reserved.
+
+#pragma once
+
+#if WITH_CEF3
+
+#if PLATFORM_WINDOWS
+#include "AllowWindowsPlatformTypes.h"
+#endif
+#include "include/cef_app.h"
+#if PLATFORM_WINDOWS
+#include "HideWindowsPlatformTypes.h"
+#endif
+
+/**
+ * Implements CEF App and other Process level interfaces 
+ */
+class FWebBrowserApp : public CefApp,
+	public CefBrowserProcessHandler,
+	public CefRenderProcessHandler
+{
+public:
+	
+	/**
+	 * Default Constructor
+	 */
+	FWebBrowserApp();
+private:
+	// CefApp methods.
+	virtual CefRefPtr<CefBrowserProcessHandler> GetBrowserProcessHandler() override { return this; }
+	virtual CefRefPtr<CefRenderProcessHandler> GetRenderProcessHandler() override { return this; }
+
+	// CefBrowserProcessHandler methods:
+	virtual void OnContextInitialized() override;
+
+	// Include the default reference counting implementation.
+	IMPLEMENT_REFCOUNTING(FWebBrowserApp);
+};
+#endif
diff --git a/Engine/Source/Developer/WebBrowser/Private/WebBrowserHandler.cpp b/Engine/Source/Developer/WebBrowser/Private/WebBrowserHandler.cpp
new file mode 100644
--- /dev/null
+++ b/Engine/Source/Developer/WebBrowser/Private/WebBrowserHandler.cpp
@@ -0,0 +1,115 @@
+// Copyright 1998-2014 Epic Games, Inc. All Rights Reserved.
+
+#include "WebBrowserPrivatePCH.h"
+#include "WebBrowserHandler.h"
+#include "WebBrowserWindow.h"
+#include "WebBrowserSingleton.h"
+
+#define LOCTEXT_NAMESPACE "WebBrowserHandler"
+
+#if WITH_CEF3
+FWebBrowserHandler::FWebBrowserHandler()
+{
+}
+
+void FWebBrowserHandler::OnTitleChange(CefRefPtr<CefBrowser> Browser, const CefString& Title)
+{
+	if (BrowserWindow.IsValid())
+	{
+		BrowserWindow.Pin()->SetTitle(Title);
+	}
+}
+
+void FWebBrowserHandler::OnAfterCreated(CefRefPtr<CefBrowser> Browser)
+{
+	if (BrowserWindow.IsValid())
+	{
+		BrowserWindow.Pin()->BindCefBrowser(Browser);
+	}
+}
+
+void FWebBrowserHandler::OnBeforeClose(CefRefPtr<CefBrowser> Browser)
+{
+	if (BrowserWindow.IsValid())
+	{
+		BrowserWindow.Pin()->BindCefBrowser(nullptr);
+	}
+}
+
+void FWebBrowserHandler::OnLoadError(CefRefPtr<CefBrowser> Browser,
+	CefRefPtr<CefFrame> Frame,
+	CefLoadHandler::ErrorCode InErrorCode,
+	const CefString& ErrorText,
+	const CefString& FailedUrl)
+{
+	// Don't display an error for downloaded files.
+	if (InErrorCode == ERR_ABORTED)
+		return;
+
+	// Display a load error message.
+	FFormatNamedArguments Args;
+	Args.Add(TEXT("FailedUrl"), FText::FromString(FailedUrl.c_str()));
+	Args.Add(TEXT("ErrorText"), FText::FromString(ErrorText.c_str()));
+	Args.Add(TEXT("ErrorCode"), FText::AsNumber(InErrorCode));
+	FText ErrorMsg = FText::Format(LOCTEXT("WebBrowserLoadError", "Failed to load URL {FailedUrl} with error {ErrorText} ({ErrorCode})."), Args);
+	FString ErrorHTML = TEXT("<html><body bgcolor=\"white\"><h2>")
+						+ ErrorMsg.ToString()
+						+ TEXT("</h2></body></html>");
+	Frame->LoadString(*ErrorHTML, FailedUrl);
+}
+
+bool FWebBrowserHandler::GetViewRect(CefRefPtr<CefBrowser> Browser, CefRect& Rect)
+{
+	if (BrowserWindow.IsValid())
+	{
+		return BrowserWindow.Pin()->GetViewRect(Rect);
+	}
+	return false;
+}
+
+void FWebBrowserHandler::OnPaint(CefRefPtr<CefBrowser> Browser,
+	PaintElementType Type,
+	const RectList& DirtyRects,
+	const void* Buffer,
+	int Width, int Height)
+{
+	if (BrowserWindow.IsValid())
+	{
+		BrowserWindow.Pin()->OnPaint(Type, DirtyRects, Buffer, Width, Height);
+	}
+}
+
+void FWebBrowserHandler::OnCursorChange(CefRefPtr<CefBrowser> Browser, CefCursorHandle Cursor)
+{
+	if (BrowserWindow.IsValid())
+	{
+		BrowserWindow.Pin()->OnCursorChange(Cursor);
+	}
+}
+
+bool FWebBrowserHandler::OnBeforeResourceLoad(CefRefPtr<CefBrowser> Browser, CefRefPtr<CefFrame> Frame, CefRefPtr<CefRequest> Request)
+{
+	const FString LanguageHeaderText(TEXT("Accept-Language"));
+	const FString LocaleCode = FWebBrowserSingleton::GetCurrentLocaleCode();
+	CefRequest::HeaderMap HeaderMap;
+	Request->GetHeaderMap(HeaderMap);
+	auto LanguageHeader = HeaderMap.find(*LanguageHeaderText);
+	if (LanguageHeader != HeaderMap.end())
+	{
+		(*LanguageHeader).second = *LocaleCode;
+	}
+	else
+	{
+		HeaderMap.insert(std::pair<CefString, CefString>(*LanguageHeaderText, *LocaleCode));
+	}
+	Request->SetHeaderMap(HeaderMap);
+	return false;
+}
+
+void FWebBrowserHandler::SetBrowserWindow(TSharedPtr<FWebBrowserWindow> InBrowserWindow)
+{
+	BrowserWindow = InBrowserWindow;
+}
+#endif
+
+#undef LOCTEXT_NAMESPACE
diff --git a/Engine/Source/Developer/WebBrowser/Private/WebBrowserHandler.h b/Engine/Source/Developer/WebBrowser/Private/WebBrowserHandler.h
new file mode 100644
--- /dev/null
+++ b/Engine/Source/Developer/WebBrowser/Private/WebBrowserHandler.h
@@ -0,0 +1,98 @@
+// Copyright 1998-2014 Epic Games, Inc. All Rights Reserved.
+
+#pragma once
+
+#if WITH_CEF3
+
+#if PLATFORM_WINDOWS
+#include "AllowWindowsPlatformTypes.h"
+#endif
+#include "include/cef_client.h"
+#if PLATFORM_WINDOWS
+#include "HideWindowsPlatformTypes.h"
+#endif
+
+// Forward Declarations
+class FWebBrowserWindow;
+
+/**
+ * Implements CEF Client and other Browser level interfaces 
+ */
+class FWebBrowserHandler : public CefClient,
+	public CefDisplayHandler,
+	public CefLifeSpanHandler,
+	public CefLoadHandler,
+	public CefRenderHandler,
+	public CefRequestHandler
+{
+public:
+	/**
+	 * Default Constructor
+	 */
+	FWebBrowserHandler();
+
+	// CefClient Interface
+	virtual CefRefPtr<CefDisplayHandler> GetDisplayHandler() override
+	{
+		return this;
+	}
+	virtual CefRefPtr<CefLifeSpanHandler> GetLifeSpanHandler() override
+	{
+		return this;
+	}
+	virtual CefRefPtr<CefLoadHandler> GetLoadHandler() override
+	{
+		return this;
+	}
+	virtual CefRefPtr<CefRenderHandler> GetRenderHandler() override
+	{
+		return this;
+	}
+	virtual CefRefPtr<CefRequestHandler> GetRequestHandler() override
+	{
+		return this;
+	}
+
+	// CefDisplayHandler Interface
+	virtual void OnTitleChange(CefRefPtr<CefBrowser> Browser, const CefString& Title) override;
+
+	// CefLifeSpanHandler Interface
+	virtual void OnAfterCreated(CefRefPtr<CefBrowser> Browser) override;
+	virtual void OnBeforeClose(CefRefPtr<CefBrowser> Browser) override;
+
+	// CefLoadHandler Interface
+	virtual void OnLoadError(CefRefPtr<CefBrowser> Browser,
+		CefRefPtr<CefFrame> Frame,
+		CefLoadHandler::ErrorCode InErrorCode,
+		const CefString& ErrorText,
+		const CefString& FailedUrl) override;
+
+	// CefRenderHandler Interface
+	virtual bool GetViewRect(CefRefPtr<CefBrowser> Browser, CefRect& Rect) override;
+	virtual void OnPaint(CefRefPtr<CefBrowser> Browser,
+		PaintElementType Type,
+		const RectList& DirtyRects,
+		const void* Buffer,
+		int Width, int Height) override;
+	virtual void OnCursorChange(CefRefPtr<CefBrowser> Browser, CefCursorHandle Cursor) override;
+
+	// CefRequestHandler Interface
+	virtual bool OnBeforeResourceLoad(CefRefPtr<CefBrowser> Browser,
+		CefRefPtr<CefFrame> Frame,
+		CefRefPtr<CefRequest> Request) override;
+
+	/**
+	 * Pass in a pointer to our Browser Window so that events can be passed on
+	 *
+	 * @param InBrowserWindow The browser window this will be handling
+	 */
+	void SetBrowserWindow(TSharedPtr<FWebBrowserWindow> InBrowserWindow);
+
+private:
+	/** Weak Pointer to our Web Browser window so that events can be passed on while it's valid*/
+	TWeakPtr<FWebBrowserWindow> BrowserWindow;
+
+	// Include the default reference counting implementation.
+	IMPLEMENT_REFCOUNTING(FWebBrowserHandler);
+};
+#endif
diff --git a/Engine/Source/Developer/WebBrowser/Private/WebBrowserModule.cpp b/Engine/Source/Developer/WebBrowser/Private/WebBrowserModule.cpp
new file mode 100644
--- /dev/null
+++ b/Engine/Source/Developer/WebBrowser/Private/WebBrowserModule.cpp
@@ -0,0 +1,47 @@
+// Copyright 1998-2014 Epic Games, Inc. All Rights Reserved.
+
+#include "WebBrowserPrivatePCH.h"
+#include "WebBrowserSingleton.h"
+#include "ModuleManager.h"
+#include "CEF3Utils.h"
+
+DEFINE_LOG_CATEGORY(LogWebBrowser);
+
+static FWebBrowserSingleton* WebBrowserSingleton = nullptr;
+
+class FWebBrowserModule : public IWebBrowserModule
+{
+private:
+	// IModuleInterface Interface
+	virtual void StartupModule() override;
+	virtual void ShutdownModule() override;
+
+public:
+	virtual IWebBrowserSingleton* GetSingleton();
+};
+
+IMPLEMENT_MODULE( FWebBrowserModule, WebBrowser );
+
+void FWebBrowserModule::StartupModule()
+{
+#if WITH_CEF3
+	CEF3Utils::LoadCEF3Modules();
+#endif
+
+	WebBrowserSingleton = new FWebBrowserSingleton();
+}
+
+void FWebBrowserModule::ShutdownModule()
+{
+	delete WebBrowserSingleton;
+	WebBrowserSingleton = NULL;
+
+#if WITH_CEF3
+	CEF3Utils::UnloadCEF3Modules();
+#endif
+}
+
+IWebBrowserSingleton* FWebBrowserModule::GetSingleton()
+{
+	return WebBrowserSingleton;
+}
\ No newline at end of file
diff --git a/Engine/Source/Developer/WebBrowser/Private/WebBrowserPrivatePCH.h b/Engine/Source/Developer/WebBrowser/Private/WebBrowserPrivatePCH.h
new file mode 100644
--- /dev/null
+++ b/Engine/Source/Developer/WebBrowser/Private/WebBrowserPrivatePCH.h
@@ -0,0 +1,7 @@
+// Copyright 1998-2014 Epic Games, Inc. All Rights Reserved.
+
+#pragma once
+
+#include "WebBrowserModule.h"
+
+DECLARE_LOG_CATEGORY_EXTERN(LogWebBrowser, Log, All);
\ No newline at end of file
diff --git a/Engine/Source/Developer/WebBrowser/Private/WebBrowserSingleton.cpp b/Engine/Source/Developer/WebBrowser/Private/WebBrowserSingleton.cpp
new file mode 100644
--- /dev/null
+++ b/Engine/Source/Developer/WebBrowser/Private/WebBrowserSingleton.cpp
@@ -0,0 +1,133 @@
+// Copyright 1998-2014 Epic Games, Inc. All Rights Reserved.
+
+#include "WebBrowserPrivatePCH.h"
+#include "WebBrowserSingleton.h"
+#include "WebBrowserApp.h"
+#include "WebBrowserHandler.h"
+#include "WebBrowserWindow.h"
+
+#if WITH_CEF3
+#if PLATFORM_WINDOWS
+	#include "AllowWindowsPlatformTypes.h"
+#endif
+	#include "include/cef_app.h"
+#if PLATFORM_WINDOWS
+	#include "HideWindowsPlatformTypes.h"
+#endif
+#endif
+
+FWebBrowserSingleton::FWebBrowserSingleton()
+{
+#if WITH_CEF3
+	// Provide CEF with command-line arguments.
+	CefMainArgs MainArgs(hInstance);
+
+	// WebBrowserApp implements application-level callbacks.
+	WebBrowserApp = new FWebBrowserApp;
+
+	// Specify CEF global settings here.
+	CefSettings Settings;
+	Settings.no_sandbox = true;
+	Settings.command_line_args_disabled = true;
+
+	// Specify locale from our settings
+	FString LocaleCode = GetCurrentLocaleCode();
+	CefString(&Settings.locale) = *LocaleCode;
+
+	// Specify path to resources
+	FString ResourcesPath(FPaths::Combine(*FPaths::EngineDir(), TEXT("Binaries/ThirdParty/CEF3/Win64/Resources")));
+	FString LocalesPath(FPaths::Combine(*FPaths::EngineDir(), TEXT("Binaries/ThirdParty/CEF3/Win64/Resources/locales")));
+	CefString(&Settings.locales_dir_path) = *LocalesPath;
+	CefString(&Settings.resources_dir_path) = *ResourcesPath;
+
+	// Specify path to sub process exe
+	FString SubProcessPath(TEXT("UnrealCEFSubProcess.exe"));
+	CefString(&Settings.browser_subprocess_path) = *SubProcessPath;
+
+	// Initialize CEF.
+	CefInitialize(MainArgs, Settings, WebBrowserApp.get(), nullptr);
+#endif
+}
+
+FWebBrowserSingleton::~FWebBrowserSingleton()
+{
+#if WITH_CEF3
+	// Force all existing browsers to close in case any haven't been deleted
+	for (int32 Index = 0; Index < WindowInterfaces.Num(); ++Index)
+	{
+		if (WindowInterfaces[Index].IsValid())
+		{
+			WindowInterfaces[Index].Pin()->CloseBrowser();
+		}
+	}
+
+	// CefRefPtr takes care of delete
+	WebBrowserApp = nullptr;
+	// Shut down CEF.
+	CefShutdown();
+#endif
+}
+
+void FWebBrowserSingleton::SetSlateRenderer(TSharedPtr<FSlateRenderer> InSlateRenderer)
+{
+	SlateRenderer = InSlateRenderer;
+}
+
+void FWebBrowserSingleton::PumpMessages()
+{
+#if WITH_CEF3
+	// Remove any windows that have been deleted
+	for (int32 Index = WindowInterfaces.Num() - 1; Index >= 0; --Index)
+	{
+		if (!WindowInterfaces[Index].IsValid())
+		{
+			WindowInterfaces.RemoveAtSwap(Index);
+		}
+	}
+	CefDoMessageLoopWork();
+#endif
+}
+
+TSharedPtr<IWebBrowserWindow> FWebBrowserSingleton::CreateBrowserWindow(void* OSWindowHandle, FString InitialURL, uint32 Width, uint32 Height)
+{
+#if WITH_CEF3
+	// Create new window
+	TSharedPtr<FWebBrowserWindow> NewWindow(new FWebBrowserWindow(SlateRenderer, FVector2D(Width, Height)));
+
+	// WebBrowserHandler implements browser-level callbacks.
+	CefRefPtr<FWebBrowserHandler> NewHandler(new FWebBrowserHandler);
+	NewWindow->SetHandler(NewHandler);
+
+	// Information used when creating the native window.
+	CefWindowHandle WindowHandle = (CefWindowHandle)OSWindowHandle; // TODO: check this is correct for all platforms
+	CefWindowInfo WindowInfo;
+
+	// Always use off screen rendering so we can integrate with our windows
+	WindowInfo.SetAsOffScreen(WindowHandle);
+
+	// Specify CEF browser settings here.
+	CefBrowserSettings BrowserSettings;
+
+	CefString URL = *InitialURL;
+
+	// Create the CEF browser window.
+	if (CefBrowserHost::CreateBrowser(WindowInfo, NewHandler.get(), URL, BrowserSettings, NULL))
+	{
+		WindowInterfaces.Add(NewWindow);
+		return NewWindow;
+	}
+#endif
+	return NULL;
+}
+
+FString FWebBrowserSingleton::GetCurrentLocaleCode()
+{
+	FCultureRef Culture = FInternationalization::Get().GetCurrentCulture();
+	FString LocaleCode = Culture->GetTwoLetterISOLanguageName();
+	FString Country = Culture->GetCountry();
+	if (!Country.IsEmpty())
+	{
+		LocaleCode = LocaleCode + TEXT("-") + Country;
+	}
+	return LocaleCode;
+}
diff --git a/Engine/Source/Developer/WebBrowser/Private/WebBrowserSingleton.h b/Engine/Source/Developer/WebBrowser/Private/WebBrowserSingleton.h
new file mode 100644
--- /dev/null
+++ b/Engine/Source/Developer/WebBrowser/Private/WebBrowserSingleton.h
@@ -0,0 +1,58 @@
+// Copyright 1998-2014 Epic Games, Inc. All Rights Reserved.
+
+#pragma once
+
+#include "IWebBrowserSingleton.h"
+
+#if WITH_CEF3
+#if PLATFORM_WINDOWS
+#include "AllowWindowsPlatformTypes.h"
+#endif
+#include "include/internal/cef_ptr.h"
+#if PLATFORM_WINDOWS
+#include "HideWindowsPlatformTypes.h"
+#endif
+#endif
+
+// Forward Declarations
+class FWebBrowserApp;
+class FWebBrowserHandler;
+class FWebBrowserWindow;
+
+/**
+ * Implementation of singleton class that takes care of general web browser tasks
+ */
+class FWebBrowserSingleton : public IWebBrowserSingleton
+{
+public:
+	/**
+	 * Default Constructor
+	 */
+	FWebBrowserSingleton();
+	/**
+	 * Virtual Destructor
+	 */
+	virtual ~FWebBrowserSingleton();
+
+	// IWebBrowserSingleton Interface
+	virtual void SetSlateRenderer(TSharedPtr<FSlateRenderer> InSlateRenderer) override;
+	virtual void PumpMessages() override;
+	TSharedPtr<IWebBrowserWindow> CreateBrowserWindow(void* OSWindowHandle, FString InitialURL, uint32 Width, uint32 Height) override;
+
+	/**
+	 * Gets the Current Locale Code in the format CEF expects
+	 *
+	 * @return Locale code as either "xx" or "xx-YY"
+	 */
+	static FString GetCurrentLocaleCode();
+
+private:
+	/** Pointer to the Slate Renderer so that we can render web pages to textures */
+	TWeakPtr<FSlateRenderer>			SlateRenderer;
+#if WITH_CEF3
+	/** Pointer to the CEF App implementation */
+	CefRefPtr<FWebBrowserApp>			WebBrowserApp;
+	/** List of currently existing browser windows */
+	TArray<TWeakPtr<FWebBrowserWindow>>	WindowInterfaces;
+#endif
+};
diff --git a/Engine/Source/Developer/WebBrowser/Private/WebBrowserViewport.cpp b/Engine/Source/Developer/WebBrowser/Private/WebBrowserViewport.cpp
new file mode 100644
--- /dev/null
+++ b/Engine/Source/Developer/WebBrowser/Private/WebBrowserViewport.cpp
@@ -0,0 +1,110 @@
+// Copyright 1998-2014 Epic Games, Inc. All Rights Reserved.
+
+#include "WebBrowserPrivatePCH.h"
+#include "WebBrowserViewport.h"
+#include "IWebBrowserWindow.h"
+
+FIntPoint FWebBrowserViewport::GetSize() const
+{
+	return (WebBrowserWindow->GetTexture() != nullptr)
+		? FIntPoint(WebBrowserWindow->GetTexture()->GetWidth(), WebBrowserWindow->GetTexture()->GetHeight())
+		: FIntPoint();
+}
+
+FSlateShaderResource* FWebBrowserViewport::GetViewportRenderTargetTexture() const
+{
+	return WebBrowserWindow->GetTexture();
+}
+
+void FWebBrowserViewport::Tick(const FGeometry& AllottedGeometry, double InCurrentTime, float DeltaTime)
+{
+	WebBrowserWindow->SetViewportSize(AllottedGeometry.GetLocalSize());
+}
+
+bool FWebBrowserViewport::RequiresVsync() const
+{
+	return false;
+}
+
+FCursorReply FWebBrowserViewport::OnCursorQuery( const FGeometry& MyGeometry, const FPointerEvent& CursorEvent )
+{
+	// TODO: retrieve cursor type from WebBrowserWindow if we can figure that out from CefCursorHandle
+	return FCursorReply::Unhandled();
+}
+
+FReply FWebBrowserViewport::OnMouseButtonDown(const FGeometry& MyGeometry, const FPointerEvent& MouseEvent)
+{
+	// Capture mouse on left button down so that you can drag out of the viewport
+	WebBrowserWindow->OnMouseButtonDown(MyGeometry, MouseEvent);
+	if (MouseEvent.GetEffectingButton() == EKeys::LeftMouseButton && ViewportWidget.IsValid())
+	{
+		return FReply::Handled().CaptureMouse(ViewportWidget.Pin().ToSharedRef());
+	}
+	return FReply::Handled();
+}
+
+FReply FWebBrowserViewport::OnMouseButtonUp(const FGeometry& MyGeometry, const FPointerEvent& MouseEvent)
+{
+	// Release mouse capture when left button released
+	WebBrowserWindow->OnMouseButtonUp(MyGeometry, MouseEvent);
+	if (MouseEvent.GetEffectingButton() == EKeys::LeftMouseButton)
+	{
+		return FReply::Handled().ReleaseMouseCapture();
+	}
+	return FReply::Handled();
+}
+
+void FWebBrowserViewport::OnMouseEnter(const FGeometry& MyGeometry, const FPointerEvent& MouseEvent)
+{
+}
+
+void FWebBrowserViewport::OnMouseLeave(const FPointerEvent& MouseEvent)
+{
+}
+
+FReply FWebBrowserViewport::OnMouseMove(const FGeometry& MyGeometry, const FPointerEvent& MouseEvent)
+{
+	WebBrowserWindow->OnMouseMove(MyGeometry, MouseEvent);
+	return FReply::Handled();
+}
+
+FReply FWebBrowserViewport::OnMouseWheel(const FGeometry& MyGeometry, const FPointerEvent& MouseEvent)
+{
+	WebBrowserWindow->OnMouseWheel(MyGeometry, MouseEvent);
+	return FReply::Handled();
+}
+
+FReply FWebBrowserViewport::OnMouseButtonDoubleClick(const FGeometry& InMyGeometry, const FPointerEvent& InMouseEvent)
+{
+	WebBrowserWindow->OnMouseButtonDoubleClick(InMyGeometry, InMouseEvent);
+	return FReply::Handled();
+}
+
+FReply FWebBrowserViewport::OnKeyDown(const FGeometry& MyGeometry, const FKeyboardEvent& InKeyboardEvent)
+{
+	WebBrowserWindow->OnKeyDown(InKeyboardEvent);
+	return FReply::Handled();
+}
+
+FReply FWebBrowserViewport::OnKeyUp(const FGeometry& MyGeometry, const FKeyboardEvent& InKeyboardEvent)
+{
+	WebBrowserWindow->OnKeyUp(InKeyboardEvent);
+	return FReply::Handled();
+}
+
+FReply FWebBrowserViewport::OnKeyChar( const FGeometry& MyGeometry, const FCharacterEvent& InCharacterEvent )
+{
+	WebBrowserWindow->OnKeyChar(InCharacterEvent);
+	return FReply::Handled();
+}
+
+FReply FWebBrowserViewport::OnKeyboardFocusReceived(const FKeyboardFocusEvent& InKeyboardFocusEvent)
+{
+	WebBrowserWindow->OnFocus(true);
+	return FReply::Handled();
+}
+
+void FWebBrowserViewport::OnKeyboardFocusLost(const FKeyboardFocusEvent& InKeyboardFocusEvent)
+{
+	WebBrowserWindow->OnFocus(false);
+}
diff --git a/Engine/Source/Developer/WebBrowser/Private/WebBrowserWindow.cpp b/Engine/Source/Developer/WebBrowser/Private/WebBrowserWindow.cpp
new file mode 100644
--- /dev/null
+++ b/Engine/Source/Developer/WebBrowser/Private/WebBrowserWindow.cpp
@@ -0,0 +1,446 @@
+// Copyright 1998-2014 Epic Games, Inc. All Rights Reserved.
+
+#include "WebBrowserPrivatePCH.h"
+#include "WebBrowserWindow.h"
+#include "SlateCore.h"
+
+#if WITH_CEF3
+FWebBrowserWindow::FWebBrowserWindow(TWeakPtr<FSlateRenderer> InSlateRenderer, FVector2D InViewportSize)
+	: SlateRenderer(InSlateRenderer)
+	, UpdatableTexture(nullptr)
+	, ViewportSize(InViewportSize)
+	, bIsClosing(false)
+{
+	TextureData.Reserve(ViewportSize.X * ViewportSize.Y * 4);
+	TextureData.SetNumZeroed(ViewportSize.X * ViewportSize.Y * 4);
+	if (SlateRenderer.IsValid())
+	{
+		UpdatableTexture = SlateRenderer.Pin()->CreateUpdatableTexture(ViewportSize.X, ViewportSize.Y);
+	}
+}
+
+FWebBrowserWindow::~FWebBrowserWindow()
+{
+	CloseBrowser();
+
+	if (SlateRenderer.IsValid() && UpdatableTexture != nullptr)
+	{
+		SlateRenderer.Pin()->ReleaseUpdatableTexture(UpdatableTexture);
+	}
+	UpdatableTexture = nullptr;
+}
+
+void FWebBrowserWindow::SetViewportSize(FVector2D WindowSize)
+{
+	// Magic number for texture size, can't access GetMax2DTextureDimension easily
+	WindowSize = WindowSize.ClampAxes(1, 2048);
+	if (ViewportSize != WindowSize)
+	{
+		ViewportSize = WindowSize;
+		TextureData.Reset(ViewportSize.X * ViewportSize.Y * 4);
+		TextureData.SetNumZeroed(ViewportSize.X * ViewportSize.Y * 4);
+		if (UpdatableTexture != nullptr)
+		{
+			UpdatableTexture->ResizeTexture(ViewportSize.X, ViewportSize.Y);
+		}
+		if (IsValid())
+		{
+			InternalCefBrowser->GetHost()->WasResized();
+		}
+	}
+}
+
+FSlateShaderResource* FWebBrowserWindow::GetTexture()
+{
+	if (UpdatableTexture != nullptr)
+	{
+		return UpdatableTexture->GetSlateResource();
+	}
+	return nullptr;
+}
+
+bool FWebBrowserWindow::IsValid() const
+{
+	return InternalCefBrowser.get() != nullptr;
+}
+
+bool FWebBrowserWindow::IsClosing() const
+{
+	return bIsClosing;
+}
+
+FString FWebBrowserWindow::GetTitle() const
+{
+	return Title;
+}
+
+void FWebBrowserWindow::OnKeyDown(const FKeyboardEvent& InKeyboardEvent)
+{
+	if (IsValid())
+	{
+		CefKeyEvent KeyEvent;
+		KeyEvent.windows_key_code = InKeyboardEvent.GetKeyCode();
+		// TODO: Figure out whether this is a system key if we come across problems
+		/*KeyEvent.is_system_key = message == WM_SYSCHAR ||
+			message == WM_SYSKEYDOWN ||
+			message == WM_SYSKEYUP;*/
+
+		KeyEvent.type = KEYEVENT_RAWKEYDOWN;
+		KeyEvent.modifiers = GetCefKeyboardModifiers(InKeyboardEvent);
+
+		InternalCefBrowser->GetHost()->SendKeyEvent(KeyEvent);
+	}
+}
+
+void FWebBrowserWindow::OnKeyUp(const FKeyboardEvent& InKeyboardEvent)
+{
+	if (IsValid())
+	{
+		CefKeyEvent KeyEvent;
+		KeyEvent.windows_key_code = InKeyboardEvent.GetKeyCode();
+		// TODO: Figure out whether this is a system key if we come across problems
+		/*KeyEvent.is_system_key = message == WM_SYSCHAR ||
+			message == WM_SYSKEYDOWN ||
+			message == WM_SYSKEYUP;*/
+
+		KeyEvent.type = KEYEVENT_KEYUP;
+		KeyEvent.modifiers = GetCefKeyboardModifiers(InKeyboardEvent);
+
+		InternalCefBrowser->GetHost()->SendKeyEvent(KeyEvent);
+	}
+}
+
+void FWebBrowserWindow::OnKeyChar(const FCharacterEvent& InCharacterEvent)
+{
+	if (IsValid())
+	{
+		CefKeyEvent KeyEvent;
+		KeyEvent.windows_key_code = InCharacterEvent.GetCharacter();
+		// TODO: Figure out whether this is a system key if we come across problems
+		/*KeyEvent.is_system_key = message == WM_SYSCHAR ||
+			message == WM_SYSKEYDOWN ||
+			message == WM_SYSKEYUP;*/
+
+		KeyEvent.type = KEYEVENT_CHAR;
+		KeyEvent.modifiers = GetCefInputModifiers(InCharacterEvent);
+
+		InternalCefBrowser->GetHost()->SendKeyEvent(KeyEvent);
+	}
+}
+
+void FWebBrowserWindow::OnMouseButtonDown(const FGeometry& MyGeometry, const FPointerEvent& MouseEvent)
+{
+	if (IsValid())
+	{
+		FKey Button = MouseEvent.GetEffectingButton();
+		CefBrowserHost::MouseButtonType Type =
+			(Button == EKeys::LeftMouseButton ? MBT_LEFT : (
+			Button == EKeys::RightMouseButton ? MBT_RIGHT : MBT_MIDDLE));
+
+		CefMouseEvent Event;
+		FVector2D LocalPos = MyGeometry.AbsoluteToLocal(MouseEvent.GetScreenSpacePosition());
+		Event.x = LocalPos.X;
+		Event.y = LocalPos.Y;
+		Event.modifiers = GetCefMouseModifiers(MouseEvent);
+
+		InternalCefBrowser->GetHost()->SendMouseClickEvent(Event, Type, false,1);
+	}
+}
+
+void FWebBrowserWindow::OnMouseButtonUp(const FGeometry& MyGeometry, const FPointerEvent& MouseEvent)
+{
+	if (IsValid())
+	{
+		FKey Button = MouseEvent.GetEffectingButton();
+		CefBrowserHost::MouseButtonType Type =
+			(Button == EKeys::LeftMouseButton ? MBT_LEFT : (
+			Button == EKeys::RightMouseButton ? MBT_RIGHT : MBT_MIDDLE));
+
+		CefMouseEvent Event;
+		FVector2D LocalPos = MyGeometry.AbsoluteToLocal(MouseEvent.GetScreenSpacePosition());
+		Event.x = LocalPos.X;
+		Event.y = LocalPos.Y;
+		Event.modifiers = GetCefMouseModifiers(MouseEvent);
+
+		InternalCefBrowser->GetHost()->SendMouseClickEvent(Event, Type, true, 1);
+	}
+}
+
+void FWebBrowserWindow::OnMouseButtonDoubleClick(const FGeometry& MyGeometry, const FPointerEvent& MouseEvent)
+{
+	if (IsValid())
+	{
+		FKey Button = MouseEvent.GetEffectingButton();
+		CefBrowserHost::MouseButtonType Type =
+			(Button == EKeys::LeftMouseButton ? MBT_LEFT : (
+			Button == EKeys::RightMouseButton ? MBT_RIGHT : MBT_MIDDLE));
+
+		CefMouseEvent Event;
+		FVector2D LocalPos = MyGeometry.AbsoluteToLocal(MouseEvent.GetScreenSpacePosition());
+		Event.x = LocalPos.X;
+		Event.y = LocalPos.Y;
+		Event.modifiers = GetCefMouseModifiers(MouseEvent);
+
+		InternalCefBrowser->GetHost()->SendMouseClickEvent(Event, Type, false, 2);
+	}
+}
+
+void FWebBrowserWindow::OnMouseMove(const FGeometry& MyGeometry, const FPointerEvent& MouseEvent)
+{
+	if (IsValid())
+	{
+		CefMouseEvent Event;
+		FVector2D LocalPos = MyGeometry.AbsoluteToLocal(MouseEvent.GetScreenSpacePosition());
+		Event.x = LocalPos.X;
+		Event.y = LocalPos.Y;
+		Event.modifiers = GetCefMouseModifiers(MouseEvent);
+
+		InternalCefBrowser->GetHost()->SendMouseMoveEvent(Event, false);
+	}
+}
+
+void FWebBrowserWindow::OnMouseWheel(const FGeometry& MyGeometry, const FPointerEvent& MouseEvent)
+{
+	if(IsValid())
+	{
+		// The original delta is reduced so this should bring it back to what CEF expects
+		const float SpinFactor = 120.0f;
+		const float TrueDelta = MouseEvent.GetWheelDelta() * SpinFactor;
+		CefMouseEvent Event;
+		FVector2D LocalPos = MyGeometry.AbsoluteToLocal(MouseEvent.GetScreenSpacePosition());
+		Event.x = LocalPos.X;
+		Event.y = LocalPos.Y;
+		Event.modifiers = GetCefMouseModifiers(MouseEvent);
+		
+		InternalCefBrowser->GetHost()->SendMouseWheelEvent(Event,
+															MouseEvent.IsShiftDown() ? TrueDelta : 0,
+															!MouseEvent.IsShiftDown() ? TrueDelta : 0);
+	}
+}
+
+void FWebBrowserWindow::OnFocus(bool SetFocus)
+{
+	if (IsValid())
+	{
+		InternalCefBrowser->GetHost()->SendFocusEvent(SetFocus);
+	}
+}
+
+void FWebBrowserWindow::OnCaptureLost()
+{
+	if (IsValid())
+	{
+		InternalCefBrowser->GetHost()->SendCaptureLostEvent();
+	}
+}
+
+bool FWebBrowserWindow::CanGoBack() const
+{
+	if (IsValid())
+	{
+		return InternalCefBrowser->CanGoBack();
+	}
+	return false;
+}
+
+void FWebBrowserWindow::GoBack()
+{
+	if (IsValid())
+	{
+		InternalCefBrowser->GoBack();
+	}
+}
+
+bool FWebBrowserWindow::CanGoForward() const
+{
+	if (IsValid())
+	{
+		return InternalCefBrowser->CanGoForward();
+	}
+	return false;
+}
+
+void FWebBrowserWindow::GoForward()
+{
+	if (IsValid())
+	{
+		InternalCefBrowser->GoForward();
+	}
+}
+
+bool FWebBrowserWindow::IsLoading() const
+{
+	if (IsValid())
+	{
+		return InternalCefBrowser->IsLoading();
+	}
+	return false;
+}
+
+void FWebBrowserWindow::Reload()
+{
+	if (IsValid())
+	{
+		InternalCefBrowser->Reload();
+	}
+}
+
+void FWebBrowserWindow::StopLoad()
+{
+	if (IsValid())
+	{
+		InternalCefBrowser->StopLoad();
+	}
+}
+
+void FWebBrowserWindow::SetHandler(CefRefPtr<FWebBrowserHandler> InHandler)
+{
+	if (InHandler.get())
+	{
+		Handler = InHandler;
+		Handler->SetBrowserWindow(SharedThis(this));
+	}
+}
+
+void FWebBrowserWindow::CloseBrowser()
+{
+	bIsClosing = true;
+	if (IsValid())
+	{
+		InternalCefBrowser->GetHost()->CloseBrowser(false);
+		InternalCefBrowser = nullptr;
+		Handler = nullptr;
+	}
+}
+
+void FWebBrowserWindow::BindCefBrowser(CefRefPtr<CefBrowser> Browser)
+{
+	InternalCefBrowser = Browser;
+}
+
+void FWebBrowserWindow::SetTitle(const CefString& InTitle)
+{
+	Title = InTitle.c_str();
+	OnTitleChangedDelegate.Broadcast(Title);
+}
+
+bool FWebBrowserWindow::GetViewRect(CefRect& Rect)
+{
+	Rect.x = 0;
+	Rect.y = 0;
+	Rect.width = ViewportSize.X;
+	Rect.height = ViewportSize.Y;
+
+	return true;
+}
+
+void FWebBrowserWindow::OnPaint(CefRenderHandler::PaintElementType Type, const CefRenderHandler::RectList& DirtyRects, const void* Buffer, int Width, int Height)
+{
+	const int BufferSize = Width*Height*4;
+	if (BufferSize <= TextureData.Num())
+	{
+		FMemory::Memcpy(TextureData.GetData(), Buffer, BufferSize);
+		if (UpdatableTexture != nullptr)
+		{
+			UpdatableTexture->UpdateTexture(TextureData);
+		}
+	}
+}
+
+void FWebBrowserWindow::OnCursorChange(CefCursorHandle Cursor)
+{
+	// TODO: Figure out Unreal cursor type from this,
+	// may need to reload unreal cursors to compare handles
+	//::SetCursor( Cursor );
+}
+
+int32 FWebBrowserWindow::GetCefKeyboardModifiers(const FKeyboardEvent& KeyboardEvent)
+{
+	int32 Modifiers = GetCefInputModifiers(KeyboardEvent);
+
+	const FKey Key = KeyboardEvent.GetKey();
+	if (Key == EKeys::LeftAlt ||
+		Key == EKeys::LeftCommand ||
+		Key == EKeys::LeftControl ||
+		Key == EKeys::LeftShift)
+	{
+		Modifiers |= EVENTFLAG_IS_LEFT;
+	}
+	if (Key == EKeys::RightAlt ||
+		Key == EKeys::RightCommand ||
+		Key == EKeys::RightControl ||
+		Key == EKeys::RightShift)
+	{
+		Modifiers |= EVENTFLAG_IS_RIGHT;
+	}
+	if (Key == EKeys::NumPadZero ||
+		Key == EKeys::NumPadOne ||
+		Key == EKeys::NumPadTwo ||
+		Key == EKeys::NumPadThree ||
+		Key == EKeys::NumPadFour ||
+		Key == EKeys::NumPadFive ||
+		Key == EKeys::NumPadSix ||
+		Key == EKeys::NumPadSeven ||
+		Key == EKeys::NumPadEight ||
+		Key == EKeys::NumPadNine)
+	{
+		Modifiers |= EVENTFLAG_IS_KEY_PAD;
+	}
+
+	return Modifiers;
+}
+
+int32 FWebBrowserWindow::GetCefMouseModifiers(const FPointerEvent& InMouseEvent)
+{
+	int32 Modifiers = GetCefInputModifiers(InMouseEvent);
+
+	if (InMouseEvent.IsMouseButtonDown(EKeys::LeftMouseButton))
+	{
+		Modifiers |= EVENTFLAG_LEFT_MOUSE_BUTTON;
+	}
+	if (InMouseEvent.IsMouseButtonDown(EKeys::MiddleMouseButton))
+	{
+		Modifiers |= EVENTFLAG_MIDDLE_MOUSE_BUTTON;
+	}
+	if (InMouseEvent.IsMouseButtonDown(EKeys::RightMouseButton))
+	{
+		Modifiers |= EVENTFLAG_RIGHT_MOUSE_BUTTON;
+	}
+
+	return Modifiers;
+}
+
+int32 FWebBrowserWindow::GetCefInputModifiers(const FInputEvent& InputEvent)
+{
+	int32 Modifiers = 0;
+
+	if (InputEvent.IsShiftDown())
+	{
+		Modifiers |= EVENTFLAG_SHIFT_DOWN;
+	}
+	if (InputEvent.IsControlDown())
+	{
+		Modifiers |= EVENTFLAG_CONTROL_DOWN;
+	}
+	if (InputEvent.IsAltDown())
+	{
+		Modifiers |= EVENTFLAG_ALT_DOWN;
+	}
+	if (InputEvent.IsCommandDown())
+	{
+		Modifiers |= EVENTFLAG_COMMAND_DOWN;
+	}
+	if (InputEvent.AreCapsLocked())
+	{
+		Modifiers |= EVENTFLAG_CAPS_LOCK_ON;
+	}
+	// TODO: Add function for this if necessary
+	/*if (InputEvent.AreNumsLocked())
+	{
+		Modifiers |= EVENTFLAG_NUM_LOCK_ON;
+	}*/
+
+	return Modifiers;
+}
+
+#endif
diff --git a/Engine/Source/Developer/WebBrowser/Private/WebBrowserWindow.h b/Engine/Source/Developer/WebBrowser/Private/WebBrowserWindow.h
new file mode 100644
--- /dev/null
+++ b/Engine/Source/Developer/WebBrowser/Private/WebBrowserWindow.h
@@ -0,0 +1,166 @@
+// Copyright 1998-2014 Epic Games, Inc. All Rights Reserved.
+
+#pragma once
+
+#if WITH_CEF3
+
+#include "IWebBrowserWindow.h"
+#include "WebBrowserHandler.h"
+
+#if PLATFORM_WINDOWS
+#include "AllowWindowsPlatformTypes.h"
+#endif
+#include "include/internal/cef_ptr.h"
+#include "include/cef_render_handler.h"
+#if PLATFORM_WINDOWS
+#include "HideWindowsPlatformTypes.h"
+#endif
+
+// Forward Declarations
+struct FInputEvent;
+
+class FSlateRenderer;
+class FSlateUpdatableTexture;
+
+/**
+ * Implementation of interface for dealing with a Web Browser window
+ */
+class FWebBrowserWindow : public IWebBrowserWindow, public TSharedFromThis<FWebBrowserWindow>
+{
+public:
+	/**
+	 * Default Constructor
+	 * 
+	 * @param InSlateRenderer Slate renderer to allow creation of texture to render to
+	 * @param InViewportSize Initial size of the browser window
+	 */
+	FWebBrowserWindow(TWeakPtr<FSlateRenderer> InSlateRenderer, FVector2D InViewportSize);
+	/**
+	 * Virtual Destructor
+	 */
+	virtual ~FWebBrowserWindow();
+
+	// IWebBrowserWindow Interface
+	virtual void SetViewportSize(FVector2D WindowSize) override;
+	virtual FSlateShaderResource* GetTexture() override;
+	virtual bool IsValid() const override;
+	virtual bool IsClosing() const override;
+	virtual FString GetTitle() const override;
+	virtual void OnKeyDown(const FKeyboardEvent& InKeyboardEvent) override;
+	virtual void OnKeyUp(const FKeyboardEvent& InKeyboardEvent) override;
+	virtual void OnKeyChar(const FCharacterEvent& InCharacterEvent) override;
+	virtual void OnMouseButtonDown(const FGeometry& MyGeometry, const FPointerEvent& MouseEvent) override;
+	virtual void OnMouseButtonUp(const FGeometry& MyGeometry, const FPointerEvent& MouseEvent) override;
+	virtual void OnMouseButtonDoubleClick(const FGeometry& MyGeometry, const FPointerEvent& MouseEvent) override;
+	virtual void OnMouseMove(const FGeometry& MyGeometry, const FPointerEvent& MouseEvent) override;
+	virtual void OnMouseWheel(const FGeometry& MyGeometry, const FPointerEvent& MouseEvent) override;
+	virtual void OnFocus(bool SetFocus) override;
+	virtual void OnCaptureLost() override;
+	virtual bool CanGoBack() const override;
+	virtual void GoBack() override;
+	virtual bool CanGoForward() const override;
+	virtual void GoForward() override;
+	virtual bool IsLoading() const override;
+	virtual void Reload() override;
+	virtual void StopLoad() override;
+
+	/**
+	 * Set the CEF Handler receiving browser callbacks for this window
+	 *
+	 * @param InHandler Pointer to the handler for this window
+	 */
+	void SetHandler(CefRefPtr<FWebBrowserHandler> InHandler);
+
+	/**
+	 * Close this window so that it can no longer be used
+	 */
+	void CloseBrowser();
+
+private:
+	/**
+	 * Called to pass reference to the underlying CefBrowser as this is not created at the same time
+	 * as the FWebBrowserWindow
+	 *
+	 * @param Browser The CefBrowser for this window
+	 */
+	void BindCefBrowser(CefRefPtr<CefBrowser> Browser);
+
+	/**
+	 * Sets the Title of this window
+	 * 
+	 * @param InTitle The new title of this window
+	 */
+	void SetTitle(const CefString& InTitle);
+
+	/**
+	 * Get the current proportions of this window
+	 *
+	 * @param Rect Reference to CefRect to store sizes
+	 * @return Whether Rect was set up correctly
+	 */
+	bool GetViewRect(CefRect& Rect);
+
+	/**
+	 * Called when there is an update to the rendered web page
+	 *
+	 * @param Type Paint type
+	 * @param DirtyRects List of image areas that have been changed
+	 * @param Buffer Pointer to the raw texture data
+	 * @param Width Width of the texture
+	 * @param Height Height of the texture
+	 */
+	void OnPaint(CefRenderHandler::PaintElementType Type, const CefRenderHandler::RectList& DirtyRects, const void* Buffer, int Width, int Height);
+	
+	/**
+	 * Called when cursor would change due to web browser interaction
+	 * 
+	 * @param Cursor Handle to CEF mouse cursor
+	 */
+	void OnCursorChange(CefCursorHandle Cursor);
+
+	/**
+	 * Gets the Cef Keyboard Modifiers based on a Keyboard Event
+	 * 
+	 * @param KeyboardEvent The Keyboard event
+	 * @return Bits representing keyboard modifiers
+	 */
+	static int32 GetCefKeyboardModifiers(const FKeyboardEvent& KeyboardEvent);
+
+	/**
+	 * Gets the Cef Mouse Modifiers based on a Mouse Event
+	 * 
+	 * @param InMouseEvent The Mouse event
+	 * @return Bits representing mouse modifiers
+	 */
+	static int32 GetCefMouseModifiers(const FPointerEvent& InMouseEvent);
+
+	/**
+	 * Gets the Cef Input Modifiers based on an Input Event
+	 * 
+	 * @param InputEvent The Input event
+	 * @return Bits representing input modifiers
+	 */
+	static int32 GetCefInputModifiers(const FInputEvent& InputEvent);
+
+	/** Pointer to Slate renderer to allow creation of texture to render to */
+	TWeakPtr<FSlateRenderer>		SlateRenderer;
+	/** Interface to the texture we are rendering to */
+	FSlateUpdatableTexture*			UpdatableTexture;
+	/** Temporary storage for the raw texture data */
+	TArray<uint8>					TextureData;
+	/** Pointer to the CEF Handler for this window */
+	CefRefPtr<FWebBrowserHandler>	Handler;
+	/** Pointer to the CEF Browser for this window */
+	CefRefPtr<CefBrowser>			InternalCefBrowser;
+	/** Current title of this window */
+	FString							Title;
+	/** Current size of this window */
+	FVector2D						ViewportSize;
+	/** Whether this window is closing */
+	bool							bIsClosing;
+
+	// Allow the Handler to access functions only it needs
+	friend class FWebBrowserHandler;
+};
+
+#endif
diff --git a/Engine/Source/Developer/WebBrowser/Public/IWebBrowserSingleton.h b/Engine/Source/Developer/WebBrowser/Public/IWebBrowserSingleton.h
new file mode 100644
--- /dev/null
+++ b/Engine/Source/Developer/WebBrowser/Public/IWebBrowserSingleton.h
@@ -0,0 +1,42 @@
+// Copyright 1998-2014 Epic Games, Inc. All Rights Reserved.
+
+#pragma once
+
+// Forward Declarations
+class FSlateRenderer;
+class IWebBrowserWindow;
+
+/**
+ * A singleton class that takes care of general web browser tasks
+ */
+class WEBBROWSER_API IWebBrowserSingleton
+{
+public:
+	/**
+	 * Virtual Destructor
+	 */
+	virtual ~IWebBrowserSingleton() {};
+
+	/**
+	 * The Web Browser singleton needs the slate renderer in order to render pages to textures.
+	 *
+	 * @param  InSlateRenderer The Slate renderer
+	 */
+	virtual void SetSlateRenderer(TSharedPtr<FSlateRenderer> InSlateRenderer) = 0;
+
+	/**
+	 * Update the web browsers internal message loop
+	 */
+	virtual void PumpMessages() = 0;
+
+	/**
+	 * Create a new web browser window
+	 *
+	 * @param OSWindowHandle Handle of OS Window that the browser will be displayed in (can be null)
+	 * @param InitialURL URL that the browser should initially navigate to
+	 * @param Width Initial width of the browser
+	 * @param Height Initial height of the browser
+	 * @return New Web Browser Window Interface (may be null if not supported)
+	 */
+	virtual TSharedPtr<IWebBrowserWindow> CreateBrowserWindow(void* OSWindowHandle, FString InitialURL, uint32 Width, uint32 Height) = 0;
+};
diff --git a/Engine/Source/Developer/WebBrowser/Public/IWebBrowserWindow.h b/Engine/Source/Developer/WebBrowser/Public/IWebBrowserWindow.h
new file mode 100644
--- /dev/null
+++ b/Engine/Source/Developer/WebBrowser/Public/IWebBrowserWindow.h
@@ -0,0 +1,171 @@
+// Copyright 1998-2014 Epic Games, Inc. All Rights Reserved.
+
+#pragma once
+
+// Forward Declarations
+struct FGeometry;
+struct FKeyboardEvent;
+struct FCharacterEvent;
+struct FPointerEvent;
+class FSlateShaderResource;
+
+DECLARE_MULTICAST_DELEGATE_OneParam( FOnTitleChangedDelegate, FString );
+
+/**
+ * Interface for dealing with a Web Browser window
+ */
+class IWebBrowserWindow
+{
+public:
+	/**
+	 * Set the desired size of the web browser viewport
+	 * 
+	 * @param WindowSize Desired viewport size
+	 */
+	virtual void SetViewportSize(FVector2D WindowSize) = 0;
+
+	/**
+	 * Gets interface to the texture representation of the browser
+	 *
+	 * @return A slate shader resource that can be rendered
+	 */
+	virtual FSlateShaderResource* GetTexture() = 0;
+
+	/**
+	 * Checks whether the web browser is valid and ready for use
+	 */
+	virtual bool IsValid() const = 0;
+
+	/**
+	 * Checks whether the web browser is currently being shut down
+	 */
+	virtual bool IsClosing() const = 0;
+
+	/**
+	 * Gets the current title of the Browser page
+	 */
+	virtual FString GetTitle() const = 0;
+
+	/**
+	 * Notify the browser that a key has been pressed
+	 *
+	 * @param  InKeyboardEvent  Keyboard event
+	 */
+	virtual void OnKeyDown(const FKeyboardEvent& InKeyboardEvent) = 0;
+
+	/**
+	 * Notify the browser that a key has been released
+	 *
+	 * @param  InKeyboardEvent  Keyboard event
+	 */
+	virtual void OnKeyUp(const FKeyboardEvent& InKeyboardEvent) = 0;
+
+	/**
+	 * Notify the browser of a character event
+	 *
+	 * @param InCharacterEvent Character event
+	 */
+	virtual void OnKeyChar(const FCharacterEvent& InCharacterEvent) = 0;
+
+	/**
+	 * Notify the browser that a mouse button was pressed within it
+	 *
+	 * @param MyGeometry The Geometry of the browser
+	 * @param MouseEvent Information about the input event
+	 */
+	virtual void OnMouseButtonDown(const FGeometry& MyGeometry, const FPointerEvent& MouseEvent) = 0;
+
+	/**
+	 * Notify the browser that a mouse button was released within it
+	 *
+	 * @param MyGeometry The Geometry of the browser
+	 * @param MouseEvent Information about the input event
+	 */
+	virtual void OnMouseButtonUp(const FGeometry& MyGeometry, const FPointerEvent& MouseEvent) = 0;
+
+	/**
+	 * Notify the browser of a double click event
+	 *
+	 * @param MyGeometry The Geometry of the browser
+	 * @param MouseEvent Information about the input event
+	 */
+	virtual void OnMouseButtonDoubleClick(const FGeometry& MyGeometry, const FPointerEvent& MouseEvent) = 0;
+
+	/**
+	 * Notify the browser that a mouse moved within it
+	 *
+	 * @param MyGeometry The Geometry of the browser
+	 * @param MouseEvent Information about the input event
+	 */
+	virtual void OnMouseMove(const FGeometry& MyGeometry, const FPointerEvent& MouseEvent) = 0;
+
+	/**
+	 * Called when the mouse wheel is spun
+	 *
+	 * @param MyGeometry The Geometry of the browser
+	 * @param MouseEvent Information about the input event
+	 */
+	virtual void OnMouseWheel(const FGeometry& MyGeometry, const FPointerEvent& MouseEvent) = 0;
+
+	/**
+	 * Called when browser receives/loses focus
+	 */
+	virtual void OnFocus(bool SetFocus) = 0;
+
+	/**
+	 * Called when Capture lost
+	 */
+	virtual void OnCaptureLost() = 0;
+
+	/**
+	 * Returns true if the browser can navigate backwards.
+	 */
+	virtual bool CanGoBack() const = 0;
+
+	/**
+	 * Navigate backwards.
+	 */
+	virtual void GoBack() = 0;
+
+	/**
+	 * Returns true if the browser can navigate forwards.
+	 */
+	virtual bool CanGoForward() const = 0;
+
+	/**
+	 * Navigate forwards.
+	 */
+	virtual void GoForward() = 0;
+
+	/**
+	 * Returns true if the browser is currently loading.
+	 */
+	virtual bool IsLoading() const = 0;
+
+	/**
+	 * Reload the current page.
+	 */
+	virtual void Reload() = 0;
+
+	/**
+	 * Stop loading the page.
+	 */
+	virtual void StopLoad() = 0;
+
+	/**
+	 * Accesses a delegate to allow callbacks when a browser title changes
+	 *
+	 * @return A multicast delegate that you can register with
+	 */
+	FOnTitleChangedDelegate& OnTitleChanged() {return OnTitleChangedDelegate;}
+
+protected:
+
+	/**
+	 * Virtual Destructor
+	 */
+	virtual ~IWebBrowserWindow() {};
+
+	/** Delegate for broadcasting title changes */
+	FOnTitleChangedDelegate OnTitleChangedDelegate;
+};
diff --git a/Engine/Source/Developer/WebBrowser/Public/WebBrowserModule.h b/Engine/Source/Developer/WebBrowser/Public/WebBrowserModule.h
new file mode 100644
--- /dev/null
+++ b/Engine/Source/Developer/WebBrowser/Public/WebBrowserModule.h
@@ -0,0 +1,32 @@
+// Copyright 1998-2014 Epic Games, Inc. All Rights Reserved.
+
+#pragma once
+
+#include "ModuleManager.h"
+
+// Forward Declarations
+class IWebBrowserSingleton;
+
+/**
+ * WebBrowserModule interface
+ */
+class IWebBrowserModule : public IModuleInterface
+{
+public:
+	/**
+	 * Get or load the Web Browser Module
+	 * 
+	 * @return The loaded module
+	 */
+	static inline IWebBrowserModule& Get()
+	{
+		return FModuleManager::LoadModuleChecked< IWebBrowserModule >("WebBrowser");
+	}
+
+	/**
+	 * Get the Web Browser Singleton
+	 * 
+	 * @return The Web Browser Singleton
+	 */
+	virtual IWebBrowserSingleton* GetSingleton() = 0;
+};
\ No newline at end of file
diff --git a/Engine/Source/Developer/WebBrowser/Public/WebBrowserViewport.h b/Engine/Source/Developer/WebBrowser/Public/WebBrowserViewport.h
new file mode 100644
--- /dev/null
+++ b/Engine/Source/Developer/WebBrowser/Public/WebBrowserViewport.h
@@ -0,0 +1,58 @@
+// Copyright 1998-2014 Epic Games, Inc. All Rights Reserved.
+
+#pragma once
+
+#include "SlateCore.h"
+
+// Forward Declarations
+class IWebBrowserWindow;
+
+/**
+ * A Slate viewport to display a Web Browser Window
+ */
+class WEBBROWSER_API FWebBrowserViewport : public ISlateViewport
+{
+public:
+	/**
+	 * Default Constructor
+	 * 
+	 * @param InWebBrowserWindow The Web Browser Window this viewport will display
+	 * @param InViewportWidget The Widget displaying this viewport (needed to capture mouse)
+	 */
+	FWebBrowserViewport(TSharedPtr<IWebBrowserWindow> InWebBrowserWindow, TSharedPtr<SWidget> InViewportWidget)
+		: WebBrowserWindow(InWebBrowserWindow)
+		, ViewportWidget(InViewportWidget)
+	{ }
+
+	/**
+	 * Destructor.
+	 */
+	~FWebBrowserViewport( )
+	{
+	}
+
+	// ISlateViewport interface
+	virtual FIntPoint GetSize() const override;
+	virtual FSlateShaderResource* GetViewportRenderTargetTexture() const override;
+	virtual void Tick( const FGeometry& AllottedGeometry, double InCurrentTime, float DeltaTime ) override;
+	virtual bool RequiresVsync() const override;
+	virtual FCursorReply OnCursorQuery( const FGeometry& MyGeometry, const FPointerEvent& CursorEvent ) override;
+	virtual FReply OnMouseButtonDown( const FGeometry& MyGeometry, const FPointerEvent& MouseEvent ) override;
+	virtual FReply OnMouseButtonUp( const FGeometry& MyGeometry, const FPointerEvent& MouseEvent ) override;
+	virtual void OnMouseEnter( const FGeometry& MyGeometry, const FPointerEvent& MouseEvent ) override;
+	virtual void OnMouseLeave( const FPointerEvent& MouseEvent ) override;
+	virtual FReply OnMouseMove( const FGeometry& MyGeometry, const FPointerEvent& MouseEvent ) override;
+	virtual FReply OnMouseWheel( const FGeometry& MyGeometry, const FPointerEvent& MouseEvent ) override;
+	virtual FReply OnMouseButtonDoubleClick( const FGeometry& InMyGeometry, const FPointerEvent& InMouseEvent ) override;
+	virtual FReply OnKeyDown( const FGeometry& MyGeometry, const FKeyboardEvent& InKeyboardEvent ) override;
+	virtual FReply OnKeyUp( const FGeometry& MyGeometry, const FKeyboardEvent& InKeyboardEvent ) override;
+	virtual FReply OnKeyChar( const FGeometry& MyGeometry, const FCharacterEvent& InCharacterEvent ) override;
+	virtual FReply OnKeyboardFocusReceived( const FKeyboardFocusEvent& InKeyboardFocusEvent ) override;
+	virtual void OnKeyboardFocusLost( const FKeyboardFocusEvent& InKeyboardFocusEvent ) override;
+
+private:
+	/** The web browser this viewport will display */
+	TSharedPtr<IWebBrowserWindow>	WebBrowserWindow;
+	/** The viewport widget using this interface */
+	TWeakPtr<SWidget>				ViewportWidget;
+};
diff --git a/Engine/Source/Developer/WebBrowser/WebBrowser.Build.cs b/Engine/Source/Developer/WebBrowser/WebBrowser.Build.cs
new file mode 100644
--- /dev/null
+++ b/Engine/Source/Developer/WebBrowser/WebBrowser.Build.cs
@@ -0,0 +1,29 @@
+// Copyright 1998-2014 Epic Games, Inc. All Rights Reserved.
+
+using UnrealBuildTool;
+
+public class WebBrowser : ModuleRules
+{
+	public WebBrowser(TargetInfo Target)
+	{
+		PublicIncludePaths.Add("Developer/WebBrowser/Public");
+		PrivateIncludePaths.Add("Developer/WebBrowser/Private");
+
+		PrivateDependencyModuleNames.AddRange(
+			new string[]
+			{
+				"Core",
+				"InputCore",
+				"SlateCore",
+				"CEF3Utils",
+			}
+		);
+
+		if (Target.Platform == UnrealTargetPlatform.Win64)
+		{
+			AddThirdPartyPrivateStaticDependencies(Target,
+				"CEF3"
+				);
+		}
+	}
+}
diff --git a/Engine/Source/Programs/UnrealCEFSubProcess/Private/Linux/UnrealCEFSubProcessLinux.cpp b/Engine/Source/Programs/UnrealCEFSubProcess/Private/Linux/UnrealCEFSubProcessLinux.cpp
new file mode 100644
--- /dev/null
+++ b/Engine/Source/Programs/UnrealCEFSubProcess/Private/Linux/UnrealCEFSubProcessLinux.cpp
@@ -0,0 +1,21 @@
+// Copyright 1998-2014 Epic Games, Inc. All Rights Reserved.
+
+#include "UnrealCEFSubProcess.h"
+
+/**
+ * WinMain, called when the application is started
+ */
+int main( int argc, char** argv )
+{
+	FPlatformMisc::SetGracefulTerminationHandler();
+
+#if WITH_CEF3
+	// Structure for passing command-line arguments.
+	// The definition of this structure is platform-specific.
+	CefMainArgs MainArgs();
+
+	return RunCEFSubProcess(MainArgs);
+#else
+	return 0;
+#endif
+}
\ No newline at end of file
diff --git a/Engine/Source/Programs/UnrealCEFSubProcess/Private/Mac/UnrealCEFSubProcessMac.cpp b/Engine/Source/Programs/UnrealCEFSubProcess/Private/Mac/UnrealCEFSubProcessMac.cpp
new file mode 100644
--- /dev/null
+++ b/Engine/Source/Programs/UnrealCEFSubProcess/Private/Mac/UnrealCEFSubProcessMac.cpp
@@ -0,0 +1,16 @@
+// Copyright 1998-2014 Epic Games, Inc. All Rights Reserved.
+
+#include "UnrealCEFSubProcess.h"
+
+int main(int argc, char *argv[])
+{
+#if WITH_CEF3
+	// Structure for passing command-line arguments.
+	// The definition of this structure is platform-specific.
+	CefMainArgs MainArgs();
+
+	return RunCEFSubProcess(MainArgs);
+#else
+	return 0;
+#endif
+}
diff --git a/Engine/Source/Programs/UnrealCEFSubProcess/Private/UnrealCEFSubProcess.cpp b/Engine/Source/Programs/UnrealCEFSubProcess/Private/UnrealCEFSubProcess.cpp
new file mode 100644
--- /dev/null
+++ b/Engine/Source/Programs/UnrealCEFSubProcess/Private/UnrealCEFSubProcess.cpp
@@ -0,0 +1,26 @@
+// Copyright 1998-2014 Epic Games, Inc. All Rights Reserved.
+
+//
+// UnrealCEFSubProcess.cpp : Defines the entry point for the browser sub process
+//
+
+#include "UnrealCEFSubProcess.h"
+
+#include "RequiredProgramMainCPPInclude.h"
+
+#include "CEF3Utils.h"
+
+//#define DEBUG_USING_CONSOLE	0
+
+IMPLEMENT_APPLICATION(UnrealCEFSubProcess, "UnrealCEFSubProcess")
+
+#if WITH_CEF3
+int32 RunCEFSubProcess(const CefMainArgs& MainArgs)
+{
+	CEF3Utils::LoadCEF3Modules();
+	// Execute the sub-process logic. This will block until the sub-process should exit.
+	int32 Result = CefExecuteProcess(MainArgs, nullptr, nullptr);
+	CEF3Utils::UnloadCEF3Modules();
+	return Result;
+}
+#endif
diff --git a/Engine/Source/Programs/UnrealCEFSubProcess/Private/UnrealCEFSubProcess.h b/Engine/Source/Programs/UnrealCEFSubProcess/Private/UnrealCEFSubProcess.h
new file mode 100644
--- /dev/null
+++ b/Engine/Source/Programs/UnrealCEFSubProcess/Private/UnrealCEFSubProcess.h
@@ -0,0 +1,23 @@
+// Copyright 1998-2014 Epic Games, Inc. All Rights Reserved.
+
+#pragma once
+
+#include "Core.h"
+
+#if WITH_CEF3
+#if PLATFORM_WINDOWS
+#include "AllowWindowsPlatformTypes.h"
+#endif
+#include "include/cef_app.h"
+#if PLATFORM_WINDOWS
+#include "HideWindowsPlatformTypes.h"
+#endif
+
+/**
+ * The application's main function.
+ *
+ * @param MainArgs Main Arguments for the process (created differently on each platform).
+ * @return Application's exit value.
+ */
+int32 RunCEFSubProcess(const CefMainArgs& MainArgs);
+#endif
diff --git a/Engine/Source/Programs/UnrealCEFSubProcess/Private/Windows/UnrealCEFSubProcessWindows.cpp b/Engine/Source/Programs/UnrealCEFSubProcess/Private/Windows/UnrealCEFSubProcessWindows.cpp
new file mode 100644
--- /dev/null
+++ b/Engine/Source/Programs/UnrealCEFSubProcess/Private/Windows/UnrealCEFSubProcessWindows.cpp
@@ -0,0 +1,19 @@
+// Copyright 1998-2014 Epic Games, Inc. All Rights Reserved.
+
+#include "UnrealCEFSubProcess.h"
+
+/**
+ * WinMain, called when the application is started
+ */
+int WINAPI WinMain( _In_ HINSTANCE hInInstance, _In_opt_ HINSTANCE hPrevInstance, _In_ LPSTR, _In_ int nCmdShow )
+{
+#if WITH_CEF3
+	// Structure for passing command-line arguments.
+	// The definition of this structure is platform-specific.
+	CefMainArgs MainArgs(hInInstance);
+
+	return RunCEFSubProcess(MainArgs);
+#else
+	return 0;
+#endif
+}
diff --git a/Engine/Source/Programs/UnrealCEFSubProcess/UnrealCEFSubProcess.Build.cs b/Engine/Source/Programs/UnrealCEFSubProcess/UnrealCEFSubProcess.Build.cs
new file mode 100644
--- /dev/null
+++ b/Engine/Source/Programs/UnrealCEFSubProcess/UnrealCEFSubProcess.Build.cs
@@ -0,0 +1,34 @@
+// Copyright 1998-2014 Epic Games, Inc. All Rights Reserved.
+
+using UnrealBuildTool;
+using System.IO;
+
+public class UnrealCEFSubProcess : ModuleRules
+{
+	public UnrealCEFSubProcess(TargetInfo Target)
+	{
+		PublicIncludePaths.Add("Runtime/Launch/Public");
+
+		PrivateIncludePaths.AddRange(
+			new string[]
+			{
+				"Programs/UnrealCEFSubProcess/Private",
+				"Runtime/Launch/Private",					// for LaunchEngineLoop.cpp include
+			}
+		);
+
+		PrivateDependencyModuleNames.AddRange(
+			new string[]
+			{
+				"Core",
+				"Projects",
+				"CEF3Utils",
+			}
+		);
+
+		AddThirdPartyPrivateStaticDependencies(Target,
+			"CEF3"
+			);
+	}
+}
+
diff --git a/Engine/Source/Programs/UnrealCEFSubProcess/UnrealCEFSubProcess.Target.cs b/Engine/Source/Programs/UnrealCEFSubProcess/UnrealCEFSubProcess.Target.cs
new file mode 100644
--- /dev/null
+++ b/Engine/Source/Programs/UnrealCEFSubProcess/UnrealCEFSubProcess.Target.cs
@@ -0,0 +1,82 @@
+// Copyright 1998-2014 Epic Games, Inc. All Rights Reserved.
+
+using System.Collections.Generic;
+using System.IO;
+using UnrealBuildTool;
+
+public class UnrealCEFSubProcessTarget : TargetRules
+{
+	public UnrealCEFSubProcessTarget(TargetInfo Target)
+	{
+		Type = TargetType.Program;
+	}
+
+	//
+	// TargetRules interface.
+	//
+
+	public override bool GetSupportedPlatforms(ref List<UnrealTargetPlatform> OutPlatforms)
+	{
+		OutPlatforms.Add(UnrealTargetPlatform.Win64);
+		OutPlatforms.Add(UnrealTargetPlatform.Mac);
+		OutPlatforms.Add(UnrealTargetPlatform.Linux);
+		return true;
+	}
+
+	public override bool ShouldCompileMonolithic(UnrealTargetPlatform InPlatform, UnrealTargetConfiguration InConfiguration)
+	{
+		return true;
+	}
+
+	public override void SetupBinaries(
+		TargetInfo Target,
+		ref List<UEBuildBinaryConfiguration> OutBuildBinaryConfigurations,
+		ref List<string> OutExtraModuleNames
+		)
+	{
+		OutBuildBinaryConfigurations.Add(
+			new UEBuildBinaryConfiguration(InType: UEBuildBinaryType.Executable,
+											InModuleNames: new List<string>() { "UnrealCEFSubProcess" })
+			);
+	}
+
+	public override void SetupGlobalEnvironment(
+		TargetInfo Target,
+		ref LinkEnvironmentConfiguration OutLinkEnvironmentConfiguration,
+		ref CPPEnvironmentConfiguration OutCPPEnvironmentConfiguration
+		)
+	{
+		// Turn off various third party features we don't need
+
+		// Currently we force Lean and Mean mode
+		UEBuildConfiguration.bCompileLeanAndMeanUE = true;
+
+		// Currently this app is not linking against the engine, so we'll compile out references from Core to the rest of the engine
+		UEBuildConfiguration.bCompileAgainstEngine = false;
+		UEBuildConfiguration.bCompileAgainstCoreUObject = false;
+		UEBuildConfiguration.bBuildWithEditorOnlyData = true;
+
+		// Never use malloc profiling in CEFSubProcess.
+		BuildConfiguration.bUseMallocProfiler = false;
+
+		// Force all shader formats to be built and included.
+		//UEBuildConfiguration.bForceBuildShaderFormats = true;
+
+		// CEFSubProcess is a console application, not a Windows app (sets entry point to main(), instead of WinMain())
+		OutLinkEnvironmentConfiguration.bIsBuildingConsoleApplication = false;
+
+		// Do NOT produce additional console app exe
+		OutLinkEnvironmentConfiguration.bBuildAdditionalConsoleApplication = false;
+
+		// Disable logging, as the sub processes are spawned often and logging will just slow them down
+		OutCPPEnvironmentConfiguration.Definitions.Add("ALLOW_LOG_FILE=0");
+	}
+    public override bool GUBP_AlwaysBuildWithBaseEditor()
+    {
+        return true;
+    }
+    public override bool GUBP_NeedsPlatformSpecificDLLs()
+    {
+        return true;
+    }
+}
