Allowing transparency in web browser windows and improved viewport resizing

Allow user to chose whether transparency is used for individual browser windows.
When browser viewport is resized or rendered with old size afterwards, made sure that the texture is copied row by row to avoid distortion.
Store size of viewport as an FIntPoint to avoid various casts that would be needed.

diff --git a/Engine/Source/Developer/WebBrowser/Private/WebBrowserSingleton.cpp b/Engine/Source/Developer/WebBrowser/Private/WebBrowserSingleton.cpp
--- a/Engine/Source/Developer/WebBrowser/Private/WebBrowserSingleton.cpp
+++ b/Engine/Source/Developer/WebBrowser/Private/WebBrowserSingleton.cpp
@@ -88,11 +88,11 @@ void FWebBrowserSingleton::PumpMessages()
 #endif
 }
 
-TSharedPtr<IWebBrowserWindow> FWebBrowserSingleton::CreateBrowserWindow(void* OSWindowHandle, FString InitialURL, uint32 Width, uint32 Height)
+TSharedPtr<IWebBrowserWindow> FWebBrowserSingleton::CreateBrowserWindow(void* OSWindowHandle, FString InitialURL, uint32 Width, uint32 Height, bool bUseTransparency)
 {
 #if WITH_CEF3
 	// Create new window
-	TSharedPtr<FWebBrowserWindow> NewWindow(new FWebBrowserWindow(SlateRenderer, FVector2D(Width, Height)));
+	TSharedPtr<FWebBrowserWindow> NewWindow(new FWebBrowserWindow(SlateRenderer, FIntPoint(Width, Height)));
 
 	// WebBrowserHandler implements browser-level callbacks.
 	CefRefPtr<FWebBrowserHandler> NewHandler(new FWebBrowserHandler);
@@ -104,6 +104,7 @@ TSharedPtr<IWebBrowserWindow> FWebBrowserSingleton::CreateBrowserWindow(void* OS
 
 	// Always use off screen rendering so we can integrate with our windows
 	WindowInfo.SetAsOffScreen(WindowHandle);
+	WindowInfo.SetTransparentPainting(bUseTransparency);
 
 	// Specify CEF browser settings here.
 	CefBrowserSettings BrowserSettings;
diff --git a/Engine/Source/Developer/WebBrowser/Private/WebBrowserSingleton.h b/Engine/Source/Developer/WebBrowser/Private/WebBrowserSingleton.h
--- a/Engine/Source/Developer/WebBrowser/Private/WebBrowserSingleton.h
+++ b/Engine/Source/Developer/WebBrowser/Private/WebBrowserSingleton.h
@@ -37,7 +37,7 @@ public:
 	// IWebBrowserSingleton Interface
 	virtual void SetSlateRenderer(TSharedPtr<FSlateRenderer> InSlateRenderer) override;
 	virtual void PumpMessages() override;
-	TSharedPtr<IWebBrowserWindow> CreateBrowserWindow(void* OSWindowHandle, FString InitialURL, uint32 Width, uint32 Height) override;
+	TSharedPtr<IWebBrowserWindow> CreateBrowserWindow(void* OSWindowHandle, FString InitialURL, uint32 Width, uint32 Height, bool bUseTransparency) override;
 
 	/**
 	 * Gets the Current Locale Code in the format CEF expects
diff --git a/Engine/Source/Developer/WebBrowser/Private/WebBrowserWindow.cpp b/Engine/Source/Developer/WebBrowser/Private/WebBrowserWindow.cpp
--- a/Engine/Source/Developer/WebBrowser/Private/WebBrowserWindow.cpp
+++ b/Engine/Source/Developer/WebBrowser/Private/WebBrowserWindow.cpp
@@ -5,7 +5,7 @@
 #include "SlateCore.h"
 
 #if WITH_CEF3
-FWebBrowserWindow::FWebBrowserWindow(TWeakPtr<FSlateRenderer> InSlateRenderer, FVector2D InViewportSize)
+FWebBrowserWindow::FWebBrowserWindow(TWeakPtr<FSlateRenderer> InSlateRenderer, FIntPoint InViewportSize)
 	: SlateRenderer(InSlateRenderer)
 	, UpdatableTexture(nullptr)
 	, ViewportSize(InViewportSize)
@@ -33,15 +33,26 @@ FWebBrowserWindow::~FWebBrowserWindow()
 void FWebBrowserWindow::SetViewportSize(FVector2D WindowSize)
 {
 	// Magic number for texture size, can't access GetMax2DTextureDimension easily
-	WindowSize = WindowSize.ClampAxes(1, 2048);
-	if (ViewportSize != WindowSize)
+	FIntPoint ClampedWindowSize = WindowSize.ClampAxes(1, 2048).IntPoint();
+	if (ViewportSize != ClampedWindowSize)
 	{
-		ViewportSize = WindowSize;
-		TextureData.Reset(ViewportSize.X * ViewportSize.Y * 4);
+		FIntPoint OldViewportSize = MoveTemp(ViewportSize);
+		TArray<uint8> OldTextureData = MoveTemp(TextureData);
+		ViewportSize = ClampedWindowSize;
 		TextureData.SetNumZeroed(ViewportSize.X * ViewportSize.Y * 4);
+
+		// copy row by row to avoid texture distortion
+		const int32 WriteWidth = FMath::Min(OldViewportSize.X, ViewportSize.X) * 4;
+		const int32 WriteHeight = FMath::Min(OldViewportSize.Y, ViewportSize.Y);
+		for (int32 RowIndex = 0; RowIndex < WriteHeight; ++RowIndex)
+		{
+			FMemory::Memcpy(TextureData.GetData() + ViewportSize.X * RowIndex * 4, OldTextureData.GetData() + OldViewportSize.X * RowIndex * 4, WriteWidth);
+		}
+
 		if (UpdatableTexture != nullptr)
 		{
 			UpdatableTexture->ResizeTexture(ViewportSize.X, ViewportSize.Y);
+			UpdatableTexture->UpdateTexture(TextureData);
 		}
 		if (IsValid())
 		{
@@ -336,15 +347,26 @@ bool FWebBrowserWindow::GetViewRect(CefRect& Rect)
 
 void FWebBrowserWindow::OnPaint(CefRenderHandler::PaintElementType Type, const CefRenderHandler::RectList& DirtyRects, const void* Buffer, int Width, int Height)
 {
-	const int BufferSize = Width*Height*4;
-	if (BufferSize <= TextureData.Num())
+	const int32 BufferSize = Width*Height*4;
+	if (BufferSize == TextureData.Num())
 	{
 		FMemory::Memcpy(TextureData.GetData(), Buffer, BufferSize);
-		if (UpdatableTexture != nullptr)
+	}
+	else
+	{
+		// copy row by row to avoid texture distortion
+		const int32 WriteWidth = FMath::Min(Width, ViewportSize.X) * 4;
+		const int32 WriteHeight = FMath::Min(Height, ViewportSize.Y);
+		for (int32 RowIndex = 0; RowIndex < WriteHeight; ++RowIndex)
 		{
-			UpdatableTexture->UpdateTexture(TextureData);
+			FMemory::Memcpy(TextureData.GetData() + ViewportSize.X * RowIndex * 4, static_cast<const uint8*>(Buffer) + Width * RowIndex * 4, WriteWidth);
 		}
 	}
+
+	if (UpdatableTexture != nullptr)
+	{
+		UpdatableTexture->UpdateTexture(TextureData);
+	}
 }
 
 void FWebBrowserWindow::OnCursorChange(CefCursorHandle Cursor)
diff --git a/Engine/Source/Developer/WebBrowser/Private/WebBrowserWindow.h b/Engine/Source/Developer/WebBrowser/Private/WebBrowserWindow.h
--- a/Engine/Source/Developer/WebBrowser/Private/WebBrowserWindow.h
+++ b/Engine/Source/Developer/WebBrowser/Private/WebBrowserWindow.h
@@ -34,7 +34,7 @@ public:
 	 * @param InSlateRenderer Slate renderer to allow creation of texture to render to
 	 * @param InViewportSize Initial size of the browser window
 	 */
-	FWebBrowserWindow(TWeakPtr<FSlateRenderer> InSlateRenderer, FVector2D InViewportSize);
+	FWebBrowserWindow(TWeakPtr<FSlateRenderer> InSlateRenderer, FIntPoint InViewportSize);
 	/**
 	 * Virtual Destructor
 	 */
@@ -155,7 +155,7 @@ private:
 	/** Current title of this window */
 	FString							Title;
 	/** Current size of this window */
-	FVector2D						ViewportSize;
+	FIntPoint						ViewportSize;
 	/** Whether this window is closing */
 	bool							bIsClosing;
 
diff --git a/Engine/Source/Developer/WebBrowser/Public/IWebBrowserSingleton.h b/Engine/Source/Developer/WebBrowser/Public/IWebBrowserSingleton.h
--- a/Engine/Source/Developer/WebBrowser/Public/IWebBrowserSingleton.h
+++ b/Engine/Source/Developer/WebBrowser/Public/IWebBrowserSingleton.h
@@ -36,7 +36,8 @@ public:
 	 * @param InitialURL URL that the browser should initially navigate to
 	 * @param Width Initial width of the browser
 	 * @param Height Initial height of the browser
+	 * @param bUseTransparency Whether to allow transparent rendering of pages
 	 * @return New Web Browser Window Interface (may be null if not supported)
 	 */
-	virtual TSharedPtr<IWebBrowserWindow> CreateBrowserWindow(void* OSWindowHandle, FString InitialURL, uint32 Width, uint32 Height) = 0;
+	virtual TSharedPtr<IWebBrowserWindow> CreateBrowserWindow(void* OSWindowHandle, FString InitialURL, uint32 Width, uint32 Height, bool bUseTransparency) = 0;
 };
diff --git a/Engine/Source/Runtime/Slate/Private/Widgets/SWebBrowser.cpp b/Engine/Source/Runtime/Slate/Private/Widgets/SWebBrowser.cpp
--- a/Engine/Source/Runtime/Slate/Private/Widgets/SWebBrowser.cpp
+++ b/Engine/Source/Runtime/Slate/Private/Widgets/SWebBrowser.cpp
@@ -22,7 +22,11 @@ void SWebBrowser::Construct(const FArguments& InArgs)
 		TSharedPtr<FGenericWindow> NativeWindow = InArgs._ParentWindow->GetNativeWindow();
 		OSWindowHandle = NativeWindow->GetOSWindowHandle();
 	}
-	BrowserWindow = IWebBrowserModule::Get().GetSingleton()->CreateBrowserWindow(OSWindowHandle, InArgs._InitialURL, InArgs._ViewportSize.Get().X, InArgs._ViewportSize.Get().Y);
+	BrowserWindow = IWebBrowserModule::Get().GetSingleton()->CreateBrowserWindow(OSWindowHandle,
+																				InArgs._InitialURL,
+																				InArgs._ViewportSize.Get().X,
+																				InArgs._ViewportSize.Get().Y,
+																				InArgs._SupportsTransparency);
 
 	TSharedPtr<SViewport> ViewportWidget;
 
@@ -74,6 +78,8 @@ void SWebBrowser::Construct(const FArguments& InArgs)
 			SAssignNew(ViewportWidget, SViewport)
 			.ViewportSize(InArgs._ViewportSize)
 			.EnableGammaCorrection(false)
+			.EnableBlending(InArgs._SupportsTransparency)
+			.IgnoreTextureAlpha(!InArgs._SupportsTransparency)
 		]
 	];
 
diff --git a/Engine/Source/Runtime/Slate/Public/Widgets/SWebBrowser.h b/Engine/Source/Runtime/Slate/Public/Widgets/SWebBrowser.h
--- a/Engine/Source/Runtime/Slate/Public/Widgets/SWebBrowser.h
+++ b/Engine/Source/Runtime/Slate/Public/Widgets/SWebBrowser.h
@@ -11,6 +11,7 @@ public:
 	SLATE_BEGIN_ARGS(SWebBrowser)
 		: _InitialURL(TEXT("www.google.com"))
 		, _ShowControls(true)
+		, _SupportsTransparency(false)
 		, _ViewportSize(FVector2D(320, 240))
 		{}
 
@@ -23,6 +24,9 @@ public:
 		/** Whether to show standard controls like Back, Forward, Reload etc. */
 		SLATE_ARGUMENT(bool, ShowControls)
 
+		/** Should this browser window support transparency */
+		SLATE_ARGUMENT(bool, SupportsTransparency)
+
 		/** Desired size of the web browser viewport */
 		SLATE_ATTRIBUTE(FVector2D, ViewportSize);
 	SLATE_END_ARGS()
