Changing Web Browser Singleton to tick itself

Removes depency on WebBrowser project from Slate, it now depends on Slate instead.
Removed functions from Singleton interface that are no longer needed.
Get the Slate renderer when needed instead of keeping reference.
Added call to FTicker::Tick function in Slate viewer app so that Web Browser will continue to work.

diff --git a/Engine/Source/Developer/WebBrowser/Private/WebBrowserSingleton.cpp b/Engine/Source/Developer/WebBrowser/Private/WebBrowserSingleton.cpp
--- a/Engine/Source/Developer/WebBrowser/Private/WebBrowserSingleton.cpp
+++ b/Engine/Source/Developer/WebBrowser/Private/WebBrowserSingleton.cpp
@@ -45,7 +45,8 @@ FWebBrowserSingleton::FWebBrowserSingleton()
 	CefString(&Settings.browser_subprocess_path) = *SubProcessPath;
 
 	// Initialize CEF.
-	CefInitialize(MainArgs, Settings, WebBrowserApp.get(), nullptr);
+	bool bSuccess = CefInitialize(MainArgs, Settings, WebBrowserApp.get(), nullptr);
+	check(bSuccess);
 #endif
 }
 
@@ -68,31 +69,11 @@ FWebBrowserSingleton::~FWebBrowserSingleton()
 #endif
 }
 
-void FWebBrowserSingleton::SetSlateRenderer(TSharedPtr<FSlateRenderer> InSlateRenderer)
-{
-	SlateRenderer = InSlateRenderer;
-}
-
-void FWebBrowserSingleton::PumpMessages()
-{
-#if WITH_CEF3
-	// Remove any windows that have been deleted
-	for (int32 Index = WindowInterfaces.Num() - 1; Index >= 0; --Index)
-	{
-		if (!WindowInterfaces[Index].IsValid())
-		{
-			WindowInterfaces.RemoveAtSwap(Index);
-		}
-	}
-	CefDoMessageLoopWork();
-#endif
-}
-
 TSharedPtr<IWebBrowserWindow> FWebBrowserSingleton::CreateBrowserWindow(void* OSWindowHandle, FString InitialURL, uint32 Width, uint32 Height, bool bUseTransparency)
 {
 #if WITH_CEF3
 	// Create new window
-	TSharedPtr<FWebBrowserWindow> NewWindow(new FWebBrowserWindow(SlateRenderer, FIntPoint(Width, Height)));
+	TSharedPtr<FWebBrowserWindow> NewWindow(new FWebBrowserWindow(FIntPoint(Width, Height)));
 
 	// WebBrowserHandler implements browser-level callbacks.
 	CefRefPtr<FWebBrowserHandler> NewHandler(new FWebBrowserHandler);
@@ -121,6 +102,22 @@ TSharedPtr<IWebBrowserWindow> FWebBrowserSingleton::CreateBrowserWindow(void* OS
 	return NULL;
 }
 
+bool FWebBrowserSingleton::Tick(float DeltaTime)
+{
+#if WITH_CEF3
+	// Remove any windows that have been deleted
+	for (int32 Index = WindowInterfaces.Num() - 1; Index >= 0; --Index)
+	{
+		if (!WindowInterfaces[Index].IsValid())
+		{
+			WindowInterfaces.RemoveAtSwap(Index);
+		}
+	}
+	CefDoMessageLoopWork();
+#endif
+	return true;
+}
+
 FString FWebBrowserSingleton::GetCurrentLocaleCode()
 {
 	FCultureRef Culture = FInternationalization::Get().GetCurrentCulture();
diff --git a/Engine/Source/Developer/WebBrowser/Private/WebBrowserSingleton.h b/Engine/Source/Developer/WebBrowser/Private/WebBrowserSingleton.h
--- a/Engine/Source/Developer/WebBrowser/Private/WebBrowserSingleton.h
+++ b/Engine/Source/Developer/WebBrowser/Private/WebBrowserSingleton.h
@@ -22,7 +22,7 @@ class FWebBrowserWindow;
 /**
  * Implementation of singleton class that takes care of general web browser tasks
  */
-class FWebBrowserSingleton : public IWebBrowserSingleton
+class FWebBrowserSingleton : public IWebBrowserSingleton, public FTickerObjectBase
 {
 public:
 	/**
@@ -35,10 +35,11 @@ public:
 	virtual ~FWebBrowserSingleton();
 
 	// IWebBrowserSingleton Interface
-	virtual void SetSlateRenderer(TSharedPtr<FSlateRenderer> InSlateRenderer) override;
-	virtual void PumpMessages() override;
 	TSharedPtr<IWebBrowserWindow> CreateBrowserWindow(void* OSWindowHandle, FString InitialURL, uint32 Width, uint32 Height, bool bUseTransparency) override;
 
+	// FTickerObjectBase Interface
+	virtual bool Tick(float DeltaTime) override;
+
 	/**
 	 * Gets the Current Locale Code in the format CEF expects
 	 *
@@ -47,8 +48,6 @@ public:
 	static FString GetCurrentLocaleCode();
 
 private:
-	/** Pointer to the Slate Renderer so that we can render web pages to textures */
-	TWeakPtr<FSlateRenderer>			SlateRenderer;
 #if WITH_CEF3
 	/** Pointer to the CEF App implementation */
 	CefRefPtr<FWebBrowserApp>			WebBrowserApp;
diff --git a/Engine/Source/Developer/WebBrowser/Private/WebBrowserWindow.cpp b/Engine/Source/Developer/WebBrowser/Private/WebBrowserWindow.cpp
--- a/Engine/Source/Developer/WebBrowser/Private/WebBrowserWindow.cpp
+++ b/Engine/Source/Developer/WebBrowser/Private/WebBrowserWindow.cpp
@@ -3,19 +3,20 @@
 #include "WebBrowserPrivatePCH.h"
 #include "WebBrowserWindow.h"
 #include "SlateCore.h"
+#include "SlateBasics.h"
 
 #if WITH_CEF3
-FWebBrowserWindow::FWebBrowserWindow(TWeakPtr<FSlateRenderer> InSlateRenderer, FIntPoint InViewportSize)
-	: SlateRenderer(InSlateRenderer)
-	, UpdatableTexture(nullptr)
+FWebBrowserWindow::FWebBrowserWindow(FIntPoint InViewportSize)
+	: UpdatableTexture(nullptr)
 	, ViewportSize(InViewportSize)
 	, bIsClosing(false)
 {
 	TextureData.Reserve(ViewportSize.X * ViewportSize.Y * 4);
 	TextureData.SetNumZeroed(ViewportSize.X * ViewportSize.Y * 4);
-	if (SlateRenderer.IsValid())
+
+	if (FSlateApplication::Get().GetRenderer().IsValid())
 	{
-		UpdatableTexture = SlateRenderer.Pin()->CreateUpdatableTexture(ViewportSize.X, ViewportSize.Y);
+		UpdatableTexture = FSlateApplication::Get().GetRenderer()->CreateUpdatableTexture(ViewportSize.X, ViewportSize.Y);
 	}
 }
 
@@ -23,9 +24,9 @@ FWebBrowserWindow::~FWebBrowserWindow()
 {
 	CloseBrowser();
 
-	if (SlateRenderer.IsValid() && UpdatableTexture != nullptr)
+	if (FSlateApplication::Get().GetRenderer().IsValid() && UpdatableTexture != nullptr)
 	{
-		SlateRenderer.Pin()->ReleaseUpdatableTexture(UpdatableTexture);
+		FSlateApplication::Get().GetRenderer()->ReleaseUpdatableTexture(UpdatableTexture);
 	}
 	UpdatableTexture = nullptr;
 }
diff --git a/Engine/Source/Developer/WebBrowser/Private/WebBrowserWindow.h b/Engine/Source/Developer/WebBrowser/Private/WebBrowserWindow.h
--- a/Engine/Source/Developer/WebBrowser/Private/WebBrowserWindow.h
+++ b/Engine/Source/Developer/WebBrowser/Private/WebBrowserWindow.h
@@ -31,10 +31,9 @@ public:
 	/**
 	 * Default Constructor
 	 * 
-	 * @param InSlateRenderer Slate renderer to allow creation of texture to render to
 	 * @param InViewportSize Initial size of the browser window
 	 */
-	FWebBrowserWindow(TWeakPtr<FSlateRenderer> InSlateRenderer, FIntPoint InViewportSize);
+	FWebBrowserWindow(FIntPoint InViewportSize);
 	/**
 	 * Virtual Destructor
 	 */
@@ -142,8 +141,6 @@ private:
 	 */
 	static int32 GetCefInputModifiers(const FInputEvent& InputEvent);
 
-	/** Pointer to Slate renderer to allow creation of texture to render to */
-	TWeakPtr<FSlateRenderer>		SlateRenderer;
 	/** Interface to the texture we are rendering to */
 	FSlateUpdatableTexture*			UpdatableTexture;
 	/** Temporary storage for the raw texture data */
diff --git a/Engine/Source/Developer/WebBrowser/Public/IWebBrowserSingleton.h b/Engine/Source/Developer/WebBrowser/Public/IWebBrowserSingleton.h
--- a/Engine/Source/Developer/WebBrowser/Public/IWebBrowserSingleton.h
+++ b/Engine/Source/Developer/WebBrowser/Public/IWebBrowserSingleton.h
@@ -17,18 +17,6 @@ public:
 	 */
 	virtual ~IWebBrowserSingleton() {};
 
-	/**
-	 * The Web Browser singleton needs the slate renderer in order to render pages to textures.
-	 *
-	 * @param  InSlateRenderer The Slate renderer
-	 */
-	virtual void SetSlateRenderer(TSharedPtr<FSlateRenderer> InSlateRenderer) = 0;
-
-	/**
-	 * Update the web browsers internal message loop
-	 */
-	virtual void PumpMessages() = 0;
-
 	/**
 	 * Create a new web browser window
 	 *
diff --git a/Engine/Source/Developer/WebBrowser/WebBrowser.Build.cs b/Engine/Source/Developer/WebBrowser/WebBrowser.Build.cs
--- a/Engine/Source/Developer/WebBrowser/WebBrowser.Build.cs
+++ b/Engine/Source/Developer/WebBrowser/WebBrowser.Build.cs
@@ -15,6 +15,7 @@ public class WebBrowser : ModuleRules
 				"Core",
 				"InputCore",
 				"SlateCore",
+				"Slate",
 				"CEF3Utils",
 			}
 		);
diff --git a/Engine/Source/Programs/SlateViewer/Private/SlateViewerApp.cpp b/Engine/Source/Programs/SlateViewer/Private/SlateViewerApp.cpp
--- a/Engine/Source/Programs/SlateViewer/Private/SlateViewerApp.cpp
+++ b/Engine/Source/Programs/SlateViewer/Private/SlateViewerApp.cpp
@@ -49,6 +49,7 @@ void RunSlateViewer( const TCHAR* CommandLine )
 	// loop while the server does the rest
 	while (!GIsRequestingExit)
 	{
+		FTicker::GetCoreTicker().Tick(0);
 		FSlateApplication::Get().PumpMessages();
 		FSlateApplication::Get().Tick();
 		FPlatformProcess::Sleep(0);
diff --git a/Engine/Source/Runtime/Slate/Private/Framework/Application/SlateApplication.cpp b/Engine/Source/Runtime/Slate/Private/Framework/Application/SlateApplication.cpp
--- a/Engine/Source/Runtime/Slate/Private/Framework/Application/SlateApplication.cpp
+++ b/Engine/Source/Runtime/Slate/Private/Framework/Application/SlateApplication.cpp
@@ -5,10 +5,6 @@
 #include "SWindowTitleBar.h"
 #include "HittestGrid.h"
 
-#if WITH_EDITOR || IS_PROGRAM
-#include "WebBrowserModule.h"
-#include "IWebBrowserSingleton.h"
-#endif
 #include "IWidgetReflector.h"
 #include "GenericCommands.h"
 #include "NotificationManager.h"
@@ -456,11 +452,6 @@ void FSlateApplication::InitializeRenderer( TSharedRef<FSlateRenderer> InRendere
 {
 	Renderer = InRenderer;
 	Renderer->Initialize();
-
-#if WITH_EDITOR || IS_PROGRAM
-	// Pass renderer on to the Web Browser module
-	IWebBrowserModule::Get().GetSingleton()->SetSlateRenderer(InRenderer);
-#endif
 }
 
 void FSlateApplication::InitializeSound( const TSharedRef<ISlateSoundDevice>& InSlateSoundDevice )
@@ -938,11 +929,6 @@ void FSlateApplication::Tick()
 	// (so they know their size)
 	FSlateNotificationManager::Get().Tick();
 
-#if WITH_EDITOR || IS_PROGRAM
-	// Calling this after Tick gives web browser a chance to render pages to texture after size changes
-	IWebBrowserModule::Get().GetSingleton()->PumpMessages();
-#endif
-
 	// Draw all windows
 	DrawWindows();
 }
diff --git a/Engine/Source/Runtime/Slate/Slate.Build.cs b/Engine/Source/Runtime/Slate/Slate.Build.cs
--- a/Engine/Source/Runtime/Slate/Slate.Build.cs
+++ b/Engine/Source/Runtime/Slate/Slate.Build.cs
@@ -18,12 +18,6 @@ public class Slate : ModuleRules
 			}
 		);
 
-		PrivateDependencyModuleNames.AddRange(
-			new string[] {
-				"WebBrowser"
-			}
-		);
-
         if (!UEBuildConfiguration.bBuildDedicatedServer && UEBuildConfiguration.bCompileFreeType)
         {
 			AddThirdPartyPrivateStaticDependencies(Target, "FreeType2");
