Moving SWebBrowser into the WebBrowser module to reduce the amount of modules/projects that will include it automatically.

To remove its use in AppFramework and keep using it in the SlateViewer app, I needed to define its tab in AppFramework and spawn it as a nomad tab in SlateViewer.

diff --git a/Engine/Source/Developer/WebBrowser/Private/SWebBrowser.cpp b/Engine/Source/Developer/WebBrowser/Private/SWebBrowser.cpp
new file mode 100644
--- /dev/null
+++ b/Engine/Source/Developer/WebBrowser/Private/SWebBrowser.cpp
@@ -0,0 +1,226 @@
+// Copyright 1998-2014 Epic Games, Inc. All Rights Reserved.
+
+#include "WebBrowserPrivatePCH.h"
+#include "SWebBrowser.h"
+#include "SThrobber.h"
+#if WITH_EDITOR || IS_PROGRAM
+#include "WebBrowserModule.h"
+#include "IWebBrowserSingleton.h"
+#include "IWebBrowserWindow.h"
+#include "WebBrowserViewport.h"
+#endif
+
+#define LOCTEXT_NAMESPACE "WebBrowser"
+
+SWebBrowser::SWebBrowser()
+{
+}
+
+void SWebBrowser::Construct(const FArguments& InArgs)
+{
+	void* OSWindowHandle = nullptr;
+	if (InArgs._ParentWindow.IsValid())
+	{
+		TSharedPtr<FGenericWindow> NativeWindow = InArgs._ParentWindow->GetNativeWindow();
+		OSWindowHandle = NativeWindow->GetOSWindowHandle();
+	}
+
+#if WITH_EDITOR || IS_PROGRAM
+	BrowserWindow = IWebBrowserModule::Get().GetSingleton()->CreateBrowserWindow(OSWindowHandle,
+	                                                                             InArgs._InitialURL,
+	                                                                             InArgs._ViewportSize.Get().X,
+	                                                                             InArgs._ViewportSize.Get().Y,
+	                                                                             InArgs._SupportsTransparency);
+#endif
+
+	TSharedPtr<SViewport> ViewportWidget;
+
+	ChildSlot
+	[
+		SNew(SVerticalBox)
+		+SVerticalBox::Slot()
+		.AutoHeight()
+		.Padding(0, 5)
+		[
+			SNew(SHorizontalBox)
+			.Visibility(InArgs._ShowControls ? EVisibility::Visible : EVisibility::Collapsed)
+			+SHorizontalBox::Slot()
+			.AutoWidth()
+			[
+				SNew(SButton)
+				.Text(LOCTEXT("Back","Back"))
+				.IsEnabled(this, &SWebBrowser::CanGoBack)
+				.OnClicked(this, &SWebBrowser::OnBackClicked)
+			]
+			+SHorizontalBox::Slot()
+			.AutoWidth()
+			[
+				SNew(SButton)
+				.Text(LOCTEXT("Forward", "Forward"))
+				.IsEnabled(this, &SWebBrowser::CanGoForward)
+				.OnClicked(this, &SWebBrowser::OnForwardClicked)
+			]
+			+SHorizontalBox::Slot()
+			.AutoWidth()
+			[
+				SNew(SButton)
+				.Text(this, &SWebBrowser::GetReloadButtonText)
+				.OnClicked(this, &SWebBrowser::OnReloadClicked)
+			]
+			+SHorizontalBox::Slot()
+			.FillWidth(1.0f)
+			.VAlign(VAlign_Center)
+			.HAlign(HAlign_Right)
+			.Padding(5)
+			[
+				SNew(STextBlock)
+				.Text(this, &SWebBrowser::GetTitleText)
+				.Justification(ETextJustify::Right)
+			]
+		]
+		+SVerticalBox::Slot()
+		[
+			SNew(SOverlay)
+			+SOverlay::Slot()
+			[
+				SAssignNew(ViewportWidget, SViewport)
+				.ViewportSize(InArgs._ViewportSize)
+				.EnableGammaCorrection(false)
+				.EnableBlending(InArgs._SupportsTransparency)
+				.IgnoreTextureAlpha(!InArgs._SupportsTransparency)
+				.Visibility(this, &SWebBrowser::GetViewportVisibility)
+			]
+			+SOverlay::Slot()
+			.HAlign(HAlign_Center)
+			.VAlign(VAlign_Center)
+			[
+				SNew(SCircularThrobber)
+				.Radius(10.0f)
+				.ToolTipText(LOCTEXT("LoadingThrobberToolTip", "Loading page..."))
+				.Visibility(this, &SWebBrowser::GetLoadingThrobberVisibility)
+			]
+		]
+	];
+
+#if WITH_EDITOR || IS_PROGRAM
+	if (BrowserWindow.IsValid())
+	{
+		BrowserViewport = MakeShareable(new FWebBrowserViewport(BrowserWindow, ViewportWidget));
+		ViewportWidget->SetViewportInterface(BrowserViewport.ToSharedRef());
+	}
+#endif
+}
+
+FText SWebBrowser::GetTitleText() const
+{
+#if WITH_EDITOR || IS_PROGRAM
+	if (BrowserWindow.IsValid())
+	{
+		return FText::FromString(BrowserWindow->GetTitle());
+	}
+#endif
+	return LOCTEXT("InvalidWindow", "Browser Window is not valid/supported");
+}
+
+bool SWebBrowser::CanGoBack() const
+{
+#if WITH_EDITOR || IS_PROGRAM
+	if (BrowserWindow.IsValid())
+	{
+		return BrowserWindow->CanGoBack();
+	}
+#endif
+	return false;
+}
+
+FReply SWebBrowser::OnBackClicked()
+{
+#if WITH_EDITOR || IS_PROGRAM
+	if (BrowserWindow.IsValid())
+	{
+		BrowserWindow->GoBack();
+	}
+#endif
+	return FReply::Handled();
+}
+
+bool SWebBrowser::CanGoForward() const
+{
+#if WITH_EDITOR || IS_PROGRAM
+	if (BrowserWindow.IsValid())
+	{
+		return BrowserWindow->CanGoForward();
+	}
+#endif
+	return false;
+}
+
+FReply SWebBrowser::OnForwardClicked()
+{
+#if WITH_EDITOR || IS_PROGRAM
+	if (BrowserWindow.IsValid())
+	{
+		BrowserWindow->GoForward();
+	}
+#endif
+	return FReply::Handled();
+}
+
+FText SWebBrowser::GetReloadButtonText() const
+{
+	static FText Reload = LOCTEXT("Reload", "Reload");
+	static FText Stop = LOCTEXT("Stop", "Stop");
+
+#if WITH_EDITOR || IS_PROGRAM
+	if (BrowserWindow.IsValid())
+	{
+		if (BrowserWindow->IsLoading())
+		{
+			return Stop;
+		}
+	}
+#endif
+	return Reload;
+}
+
+FReply SWebBrowser::OnReloadClicked()
+{
+#if WITH_EDITOR || IS_PROGRAM
+	if (BrowserWindow.IsValid())
+	{
+		if (BrowserWindow->IsLoading())
+		{
+			BrowserWindow->StopLoad();
+		}
+		else
+		{
+			BrowserWindow->Reload();
+		}
+	}
+#endif
+	return FReply::Handled();
+}
+
+EVisibility SWebBrowser::GetViewportVisibility() const
+{
+#if WITH_EDITOR || IS_PROGRAM
+	if (BrowserWindow.IsValid() && BrowserWindow->HasBeenPainted())
+	{
+		return EVisibility::Visible;
+	}
+#endif
+	return EVisibility::Hidden;
+}
+
+EVisibility SWebBrowser::GetLoadingThrobberVisibility() const
+{
+#if WITH_EDITOR || IS_PROGRAM
+	if (BrowserWindow.IsValid() && !BrowserWindow->HasBeenPainted())
+	{
+		return EVisibility::Visible;
+	}
+#endif
+	return EVisibility::Hidden;
+}
+
+#undef LOCTEXT_NAMESPACE
diff --git a/Engine/Source/Developer/WebBrowser/Public/SWebBrowser.h b/Engine/Source/Developer/WebBrowser/Public/SWebBrowser.h
new file mode 100644
--- /dev/null
+++ b/Engine/Source/Developer/WebBrowser/Public/SWebBrowser.h
@@ -0,0 +1,100 @@
+// Copyright 1998-2014 Epic Games, Inc. All Rights Reserved.
+
+#pragma once
+
+#include "SlateBasics.h"
+
+class IWebBrowserWindow;
+class FWebBrowserViewport;
+
+class WEBBROWSER_API SWebBrowser : public SCompoundWidget
+{
+public:
+	SLATE_BEGIN_ARGS(SWebBrowser)
+		: _InitialURL(TEXT("www.google.com"))
+		, _ShowControls(true)
+		, _SupportsTransparency(false)
+		, _ViewportSize(FVector2D(320, 240))
+		{}
+
+		/** A reference to the parent window */
+		SLATE_ARGUMENT(TSharedPtr<SWindow>, ParentWindow)
+
+		/** URL that the browser will initially navigate to */
+		SLATE_ARGUMENT(FString, InitialURL)
+
+		/** Whether to show standard controls like Back, Forward, Reload etc. */
+		SLATE_ARGUMENT(bool, ShowControls)
+
+		/** Should this browser window support transparency */
+		SLATE_ARGUMENT(bool, SupportsTransparency)
+
+		/** Desired size of the web browser viewport */
+		SLATE_ATTRIBUTE(FVector2D, ViewportSize);
+	SLATE_END_ARGS()
+
+	/**
+	 * Constructor
+	 */
+	SWebBrowser();
+
+	/**
+	 * Construct the widget.
+	 *
+	 * @param InArgs  Declaration from which to construct the widget.
+	 */
+	void Construct(const FArguments& InArgs);
+
+	/**
+	 * Get the current title of the web page
+	 */
+	FText GetTitleText() const;
+
+private:
+	/**
+	 * Returns true if the browser can navigate backwards.
+	 */
+	bool CanGoBack() const;
+
+	/**
+	 * Navigate backwards.
+	 */
+	FReply OnBackClicked();
+
+	/**
+	 * Returns true if the browser can navigate forwards.
+	 */
+	bool CanGoForward() const;
+
+	/**
+	 * Navigate forwards.
+	 */
+	FReply OnForwardClicked();
+
+	/**
+	 * Get text for reload button depending on status
+	 */
+	FText GetReloadButtonText() const;
+
+	/**
+	 * Reload or stop loading
+	 */
+	FReply OnReloadClicked();
+
+	/**
+	 * Get whether the page viewport should be visible
+	 */
+	EVisibility GetViewportVisibility() const;
+
+	/**
+	 * Get whether loading throbber should be visible
+	 */
+	EVisibility GetLoadingThrobberVisibility() const;
+
+private:
+	/** Interface for dealing with a web browser window */
+	TSharedPtr<IWebBrowserWindow>	BrowserWindow;
+
+	/** Viewport interface for rendering the web page */
+	TSharedPtr<FWebBrowserViewport>	BrowserViewport;
+};
diff --git a/Engine/Source/Editor/LevelEditor/LevelEditor.Build.cs b/Engine/Source/Editor/LevelEditor/LevelEditor.Build.cs
--- a/Engine/Source/Editor/LevelEditor/LevelEditor.Build.cs
+++ b/Engine/Source/Editor/LevelEditor/LevelEditor.Build.cs
@@ -53,7 +53,7 @@ public class LevelEditor : ModuleRules
 				"TargetPlatform",
 				"EngineSettings",
 				"PropertyEditor",
-                "AppFramework",
+				"WebBrowser",
 			}
 		);
 
diff --git a/Engine/Source/Programs/SlateViewer/Private/SlateViewerApp.cpp b/Engine/Source/Programs/SlateViewer/Private/SlateViewerApp.cpp
--- a/Engine/Source/Programs/SlateViewer/Private/SlateViewerApp.cpp
+++ b/Engine/Source/Programs/SlateViewer/Private/SlateViewerApp.cpp
@@ -5,13 +5,16 @@
 #include "STestSuite.h"
 #include "ISourceCodeAccessModule.h"
 #include "SPerfSuite.h"
+#include "SDockTab.h"
+#include "SWebBrowser.h"
 
 IMPLEMENT_APPLICATION(SlateViewer, "SlateViewer");
 
+#define LOCTEXT_NAMESPACE "SlateViewer"
 
 namespace WorkspaceMenu
 {
-	TSharedRef<FWorkspaceItem> DeveloperMenu = FWorkspaceItem::NewGroup(NSLOCTEXT("SlateViewer", "DeveloperMenu", "Developer"));
+	TSharedRef<FWorkspaceItem> DeveloperMenu = FWorkspaceItem::NewGroup(LOCTEXT("DeveloperMenu", "Developer"));
 }
 
 
@@ -36,9 +39,11 @@ void RunSlateViewer( const TCHAR* CommandLine )
 #endif
 
 	// set the application name
-	FGlobalTabmanager::Get()->SetApplicationTitle(NSLOCTEXT("SlateViewer", "AppTitle", "Slate Viewer"));
+	FGlobalTabmanager::Get()->SetApplicationTitle(LOCTEXT("AppTitle", "Slate Viewer"));
 	FModuleManager::LoadModuleChecked<ISlateReflectorModule>("SlateReflector").RegisterTabSpawner(WorkspaceMenu::DeveloperMenu);
 
+	FGlobalTabmanager::Get()->RegisterNomadTabSpawner("WebBrowserTab", FOnSpawnTab::CreateStatic(&SpawnWebBrowserTab))
+		.SetDisplayName(LOCTEXT("WebBrowserTab", "Web Browser"));
 	
 	if (FParse::Param(FCommandLine::Get(), TEXT("perftest")))
 	{
@@ -70,3 +75,17 @@ void RunSlateViewer( const TCHAR* CommandLine )
 
 	FSlateApplication::Shutdown();
 }
+
+TSharedRef<SDockTab> SpawnWebBrowserTab(const FSpawnTabArgs& Args)
+{
+	return SNew(SDockTab)
+		.Label(LOCTEXT("WebBrowserTab", "Web Browser"))
+		.ToolTipText(LOCTEXT("WebBrowserTabToolTip", "Switches to the Web Browser to test its features."))
+		.TabRole(ETabRole::NomadTab)
+		[
+			SNew(SWebBrowser)
+			.ParentWindow(Args.GetOwnerWindow())
+		];
+}
+
+#undef LOCTEXT_NAMESPACE
diff --git a/Engine/Source/Programs/SlateViewer/Public/SlateViewerApp.h b/Engine/Source/Programs/SlateViewer/Public/SlateViewerApp.h
--- a/Engine/Source/Programs/SlateViewer/Public/SlateViewerApp.h
+++ b/Engine/Source/Programs/SlateViewer/Public/SlateViewerApp.h
@@ -10,3 +10,8 @@
  * Run the SlateViewer .
  */
 void RunSlateViewer(const TCHAR* Commandline);
+
+/**
+ * Spawn the contents of the web browser tab
+ */
+TSharedRef<SDockTab> SpawnWebBrowserTab(const FSpawnTabArgs& Args);
diff --git a/Engine/Source/Programs/SlateViewer/SlateViewer.Build.cs b/Engine/Source/Programs/SlateViewer/SlateViewer.Build.cs
--- a/Engine/Source/Programs/SlateViewer/SlateViewer.Build.cs
+++ b/Engine/Source/Programs/SlateViewer/SlateViewer.Build.cs
@@ -10,7 +10,7 @@ public class SlateViewer : ModuleRules
 
 		PrivateDependencyModuleNames.AddRange(
 			new string[] {
-                "AppFramework",
+				"AppFramework",
 				"Core",
 				"Projects",
 				"Slate",
@@ -18,6 +18,7 @@ public class SlateViewer : ModuleRules
 				"SlateReflector",
 				"StandaloneRenderer",
 				"SourceCodeAccess",
+				"WebBrowser",
 			}
 		);
 
diff --git a/Engine/Source/Runtime/AppFramework/AppFramework.Build.cs b/Engine/Source/Runtime/AppFramework/AppFramework.Build.cs
--- a/Engine/Source/Runtime/AppFramework/AppFramework.Build.cs
+++ b/Engine/Source/Runtime/AppFramework/AppFramework.Build.cs
@@ -17,7 +17,6 @@ public class AppFramework : ModuleRules
 		PrivateDependencyModuleNames.AddRange(
 			new string[] {
 				"InputCore",
-				"WebBrowser",
 				"SlateReflector",
 			}
 		);
diff --git a/Engine/Source/Runtime/AppFramework/Private/Widgets/SWebBrowser.cpp b/Engine/Source/Runtime/AppFramework/Private/Widgets/SWebBrowser.cpp
deleted file mode 100644
diff --git a/Engine/Source/Runtime/AppFramework/Private/Widgets/Testing/STestSuite.cpp b/Engine/Source/Runtime/AppFramework/Private/Widgets/Testing/STestSuite.cpp
--- a/Engine/Source/Runtime/AppFramework/Private/Widgets/Testing/STestSuite.cpp
+++ b/Engine/Source/Runtime/AppFramework/Private/Widgets/Testing/STestSuite.cpp
@@ -12,7 +12,6 @@
 #include "TransformCalculus3D.h"
 #include "SlateRenderTransform.h"
 #include "SlateLayoutTransform.h"
-#include "SWebBrowser.h"
 #include "SInlineEditableTextBlock.h"
 #include "STextEntryPopup.h"
 #include "SDPIScaler.h"
@@ -4994,19 +4993,6 @@ TSharedRef<SDockTab> SpawnTab(const FSpawnTabArgs& Args, FName TabIdentifier)
 				]
 			];
 	}
-	else if (TabIdentifier == FName(TEXT("WebBrowserTab")))
-	{
-		const FString DocPath(FPaths::ConvertRelativePathToFull(FPaths::Combine(*FPaths::EngineDir(), TEXT("Documentation/HTML/INT/Resources/index.html"))));
-		return SNew(SDockTab)
-			.Label(LOCTEXT("WebBrowserTab", "Web Browser"))
-			.ToolTipText(LOCTEXT("WebBrowserTabToolTip", "Switches to the Web Browser to test its features."))
-			[
-				SNew(SWebBrowser)
-				.ParentWindow(Args.GetOwnerWindow())
-				//.InitialURL(DocPath)
-				//.InitialURL(TEXT("www.youtube.com"))
-			];
-	}
 #endif //WITH_FANCY_TEXT
 	else if (TabIdentifier == FName("LayoutRoundingTab"))
 	{
@@ -5245,10 +5231,6 @@ TSharedRef<SDockTab> SpawnTestSuite1( const FSpawnTabArgs& Args )
 		.SetDisplayName( NSLOCTEXT("TestSuite1", "RichEditableTextTab", "Rich Editable Text Test") )
 		.SetGroup(TestSuiteMenu::SuiteTabs);
 
-	TestSuite1TabManager->RegisterTabSpawner("WebBrowserTab", FOnSpawnTab::CreateStatic(&SpawnTab, FName("WebBrowserTab")))
-		.SetDisplayName(NSLOCTEXT("TestSuite1", "WebBrowserTab", "Web Browser test"))
-		.SetGroup(TestSuiteMenu::SuiteTabs);
-
 	TestSuite1TabManager->RegisterTabSpawner( "LayoutRoundingTab", FOnSpawnTab::CreateStatic( &SpawnTab, FName("LayoutRoundingTab") ) )
 		.SetDisplayName( NSLOCTEXT("TestSuite1", "LayoutRoundingTab", "Layout Rounding") )
 		.SetGroup(TestSuiteMenu::SuiteTabs);
diff --git a/Engine/Source/Runtime/AppFramework/Public/Widgets/Web/SWebBrowser.h b/Engine/Source/Runtime/AppFramework/Public/Widgets/Web/SWebBrowser.h
deleted file mode 100644
