Fix assert when shutting down SWebBrowser

diff --git a/Engine/Source/Developer/WebBrowser/Private/WebBrowserWindow.cpp b/Engine/Source/Developer/WebBrowser/Private/WebBrowserWindow.cpp
--- a/Engine/Source/Developer/WebBrowser/Private/WebBrowserWindow.cpp
+++ b/Engine/Source/Developer/WebBrowser/Private/WebBrowserWindow.cpp
@@ -25,7 +25,7 @@ FWebBrowserWindow::~FWebBrowserWindow()
 {
 	CloseBrowser();
 
-	if (FSlateApplication::Get().GetRenderer().IsValid() && UpdatableTexture != nullptr)
+	if (FSlateApplication::IsInitialized() && FSlateApplication::Get().GetRenderer().IsValid() && UpdatableTexture != nullptr)
 	{
 		FSlateApplication::Get().GetRenderer()->ReleaseUpdatableTexture(UpdatableTexture);
 	}
